[HairCardSolver] Move plugin from Restricted to Experimental folder

diff --git a/Engine/Plugins/Experimental/HairCardGenerator/Config/Editor.ini b/Engine/Plugins/Experimental/HairCardGenerator/Config/Editor.ini
new file mode 100644
--- /dev/null
+++ b/Engine/Plugins/Experimental/HairCardGenerator/Config/Editor.ini
@@ -0,0 +1,4 @@
+[/Script/HairCardGeneratorEditor.HairCardGeneratorEditorSettings]
+HairCardAssetNameFormat = "{groom_name}_LOD{lod}"
+HairCardAssetPathFormat = "{groom_dir}/{groom_name}_LOD{lod}"
+HairCardTextureNameFormats=((Depth, "{mesh_name}_depth"), (Tangent, "{mesh_name}_tangents"), (Attribute, "{mesh_name}_attributes"), (Coverage, "{mesh_name}_coverage"))
diff --git a/Engine/Plugins/Experimental/HairCardGenerator/HairCardGenerator.uplugin b/Engine/Plugins/Experimental/HairCardGenerator/HairCardGenerator.uplugin
new file mode 100644
--- /dev/null
+++ b/Engine/Plugins/Experimental/HairCardGenerator/HairCardGenerator.uplugin
@@ -0,0 +1,58 @@
+{
+	"FileVersion": 3,
+	"Version": 1,
+	"VersionName": "1.0",
+	"FriendlyName": "Hair Card Generator",
+	"Description": "Procedurally generate hair cards from hair strands",
+	"Category": "Geometry",
+	"CreatedBy": "Epic Games, Inc.",
+	"CreatedByURL": "https://epicgames.com",
+	"DocsURL": "",
+	"MarketplaceURL": "",
+	"SupportURL": "https://forums.unrealengine.com/",
+	"EnabledByDefault": false,
+	"CanContainContent": true,
+	"IsBetaVersion": false,
+	"IsExperimentalVersion": true,
+	"Installed": false,
+	"Modules": [
+		{
+			"Name": "HairCardGeneratorEditor",
+			"Type": "Editor",
+			"LoadingPhase": "Default",
+			"PlatformAllowList": [
+				"Win64",
+				"Linux"
+			]
+		}
+	],
+	"Plugins": [
+		{
+			"Name": "HairStrands",
+			"Enabled": true
+		},
+		{
+			"Name": "PythonFoundationPackages",
+			"Enabled": true
+		}
+	],
+	"PythonRequirements":
+	[
+		{
+			"Platform": "All",
+			"Requirements": 
+			[
+				"addict==2.4.0",
+				"glcontext==2.3.6",
+				"joblib==1.1.0",
+				"moderngl==5.6.4",
+				"multipledispatch==0.6.0",
+				"PyOpenGL==3.1.6",
+				"pyrr==0.10.3",
+				"pyyaml==6.0",
+				"scikit_learn==1.1.1",
+				"six==1.16.0"
+			]
+		}
+	]
+}
\ No newline at end of file
diff --git a/Engine/Plugins/Experimental/HairCardGenerator/Source/HairCardGeneratorEditor/HairCardGeneratorEditor.Build.cs b/Engine/Plugins/Experimental/HairCardGenerator/Source/HairCardGeneratorEditor/HairCardGeneratorEditor.Build.cs
new file mode 100644
--- /dev/null
+++ b/Engine/Plugins/Experimental/HairCardGenerator/Source/HairCardGeneratorEditor/HairCardGeneratorEditor.Build.cs
@@ -0,0 +1,42 @@
+// Copyright Epic Games, Inc. All Rights Reserved.
+
+using UnrealBuildTool;
+using System.IO; // for Path
+public class HairCardGeneratorEditor : ModuleRules
+{
+	public HairCardGeneratorEditor(ReadOnlyTargetRules Target) : base(Target)
+	{
+		PublicDependencyModuleNames.AddRange(
+			new string[]
+			{
+				"Core",
+				"CoreUObject",
+				"Engine",
+				"HairCardGeneratorFramework",
+				// ... add other public dependencies that you statically link with here ...
+			}
+		);
+
+		PrivateDependencyModuleNames.AddRange(
+			new string[]
+			{
+				"HairStrandsCore", // for FHairDescription
+				"MeshDescription",
+				"StaticMeshDescription",
+				"AssetRegistry",
+				"UnrealEd",
+				"Slate",
+				"SlateCore",
+				"EditorStyle",
+				"MainFrame",
+				"ApplicationCore",
+				"InputCore",
+				"ToolWidgets",
+				"Eigen",
+				"Json",
+				"JsonUtilities",
+				// ... add private dependencies that you statically link with here ...	
+			}
+		);
+	}
+}
diff --git a/Engine/Plugins/Experimental/HairCardGenerator/Source/HairCardGeneratorEditor/Private/HairCardGenCardSubdivider.cpp b/Engine/Plugins/Experimental/HairCardGenerator/Source/HairCardGeneratorEditor/Private/HairCardGenCardSubdivider.cpp
new file mode 100644
--- /dev/null
+++ b/Engine/Plugins/Experimental/HairCardGenerator/Source/HairCardGeneratorEditor/Private/HairCardGenCardSubdivider.cpp
@@ -0,0 +1,195 @@
+// Copyright Epic Games, Inc. All Rights Reserved.
+
+#include "HairCardGenCardSubdivider.h"
+
+FHairCardGenCardSubdivider::FHairCardGenCardSubdivider(const float SubdTolerance, const bool UseAdaptive, const int MaxSubdivisions):
+	SubdToleranceInitial(SubdTolerance), UseAdaptive(UseAdaptive), MaxSubdivisions(MaxSubdivisions){}
+
+void FHairCardGenCardSubdivider::SetSubdTolerance(const float NewSubdTolerance)
+{
+	SubdToleranceInitial = NewSubdTolerance;
+}
+
+MatrixXf FHairCardGenCardSubdivider::TArrayToMatrixXf(const TArray<float> Points)
+{
+	int N = Points.Num() / 3;
+	MatrixXf PointMatrix(N, 3);
+
+	for (int i = 0; i < N; i++) {
+		for (int j = 0; j < 3; j++)
+		{
+			PointMatrix(i, j) = Points[3 * i + j];
+		}
+	}
+
+	return PointMatrix;
+}
+
+TArray<float> FHairCardGenCardSubdivider::MatrixXfToTArray(const MatrixXf PointMatrix)
+{
+	int N = PointMatrix.rows();
+
+	TArray<float> Points;
+	Points.SetNum(3 * N);
+
+	for (int i = 0; i < N; i++) {
+		for (int j = 0; j < 3; j++)
+		{
+			Points[3 * i + j] = PointMatrix(i, j);
+		}
+	}
+
+	return Points;
+}
+
+float FHairCardGenCardSubdivider::GetCurveLength(const MatrixXf Points)
+{
+	float TotalLength = 0.0;
+
+	for (int i = 0; i < Points.rows() - 1; ++i)
+	{
+		TotalLength += (Points.row(i + 1) - Points.row(i)).norm();
+	}
+
+	return TotalLength;
+}
+
+float FHairCardGenCardSubdivider::GetAverageCurvatureRadius(const MatrixXf Points)
+{
+	int N = Points.rows();
+	float AverageCurvatureRadius = 0.0;
+
+	for(int i=0; i<(N - 2); ++i)
+	{
+		Vector3f P0 = Points.row(i);
+		Vector3f P1 = Points.row(i + 1);
+		Vector3f P2 = Points.row(i + 2);
+
+		Vector3f T = P1 - P0;
+		Vector3f U = P2 - P0;
+		Vector3f V = P2 - P1;
+
+		float TriangleArea = 2.0 * (T.cross(U)).norm() + 1e-10;
+
+		AverageCurvatureRadius += T.norm() * U.norm() * V.norm() / TriangleArea;
+	}
+
+	return AverageCurvatureRadius / (N - 2);
+}
+
+int FHairCardGenCardSubdivider::NextSubdPoint(const int IndStart)
+{
+	int IndEnd = IndStart + 1;
+	float CurveDistance = 0;
+	float LinearDistance = 0;
+
+	while (IndEnd < NumInterpPoints - 1)
+	{
+		LinearDistance = (MeanCurve.row(IndEnd) - MeanCurve.row(IndStart)).norm();
+		CurveDistance += (MeanCurve.row(IndEnd) - MeanCurve.row(IndEnd - 1)).norm();
+
+		if (CurveDistance - LinearDistance > SubdTolerance)
+		{
+			return IndEnd;
+		}
+		IndEnd++;
+	}
+
+	LastDist = CurveDistance - LinearDistance;
+	return -1;
+}
+
+TArray<float> FHairCardGenCardSubdivider::AdaptiveIter()
+{
+	TArray<float> SubdPoints;
+	SubdPoints.Reserve(60);
+
+	int Ind = 0;
+
+	do {
+		for (float x : MeanCurve.row(Ind)) SubdPoints.Add(x);
+		Ind = NextSubdPoint(Ind);
+	} while (Ind > 0);
+
+	for (float x : MeanCurve.row(NumInterpPoints - 1)) SubdPoints.Add(x);
+
+	return SubdPoints;
+}
+
+TArray<float> FHairCardGenCardSubdivider::GetAdaptiveSubdPoints()
+{
+	TArray<float> SubdPointsNew, SubdPoints = AdaptiveIter();
+	int NumPoints = SubdPoints.Num() / 3;
+
+	LastDist = 0.;
+
+	while (LastDist < 0.9 * SubdTolerance)
+	{
+		SubdTolerance *= 0.95;
+		SubdPointsNew = AdaptiveIter();
+		if (SubdPointsNew.Num() / 3 != NumPoints)
+		{
+			return SubdPoints;
+		}
+		SubdPoints = SubdPointsNew;
+	}
+
+	return SubdPoints;
+}
+
+MatrixXf FHairCardGenCardSubdivider::LinearInterpolation(const MatrixXf Points, const int NumPoints)
+{
+	MatrixXf InterpMat(NumPoints, 3);
+
+	int InputSize = Points.rows();
+
+	for (int j = 0; j < 3; j++)
+	{
+		InterpMat(0, j) = Points(0, j);
+		InterpMat(NumPoints - 1, j) = Points(InputSize - 1, j);
+
+		for (int i = 1; i < NumPoints - 1; i++){
+			float IndOutF = (float)(i * (InputSize - 1)) / (float)(NumPoints - 1);
+
+			if (fabs(ceilf(IndOutF) - IndOutF) < 1e-5f)
+			{
+				InterpMat(i, j) = Points((int)ceil(IndOutF), j);
+			}
+			else
+			{
+				int IndOutI = (int)floor(IndOutF);
+				InterpMat(i, j) = Points(IndOutI, j) + (Points(IndOutI + 1, j) - Points(IndOutI, j)) * (IndOutF - (float)IndOutI);
+			}
+		}
+	}
+
+	return InterpMat;
+}
+
+TArray<float> FHairCardGenCardSubdivider::GetSubdivisionPoints(const MatrixXf Points)
+{
+	TArray<float> SubdPoints;
+
+	if (UseAdaptive)
+	{
+		SubdTolerance = SubdToleranceInitial;
+		MeanCurve = LinearInterpolation(Points, NumInterpPoints);
+		SubdPoints = GetAdaptiveSubdPoints();
+	}
+	else
+	{
+		SubdPoints = MatrixXfToTArray(LinearInterpolation(Points, MaxSubdivisions + 1));
+	}
+
+	return SubdPoints;
+}
+
+float FHairCardGenCardSubdivider::GetToleranceFromLengthAndCurvRadius(float Length, float CurvRadius, float Subdivisions)
+{
+	if (Subdivisions < 1.01) return 1e10;
+
+	float Num = Length * Length * Length;
+	float Den = 3 * CurvRadius * CurvRadius * (2 * Subdivisions - 1) * (2 * Subdivisions - 1) * (2 * Subdivisions - 1);
+
+	return Num / Den;
+}
\ No newline at end of file
diff --git a/Engine/Plugins/Experimental/HairCardGenerator/Source/HairCardGeneratorEditor/Private/HairCardGenCardSubdivider.h b/Engine/Plugins/Experimental/HairCardGenerator/Source/HairCardGeneratorEditor/Private/HairCardGenCardSubdivider.h
new file mode 100644
--- /dev/null
+++ b/Engine/Plugins/Experimental/HairCardGenerator/Source/HairCardGeneratorEditor/Private/HairCardGenCardSubdivider.h
@@ -0,0 +1,46 @@
+// Copyright Epic Games, Inc. All Rights Reserved.
+#pragma once
+
+#ifndef EIGEN_MPL2_ONLY
+#define EIGEN_MPL2_ONLY
+#endif
+
+#include "Containers/Array.h"
+
+PRAGMA_DEFAULT_VISIBILITY_START
+THIRD_PARTY_INCLUDES_START
+#include <Eigen/Dense>
+THIRD_PARTY_INCLUDES_END
+PRAGMA_DEFAULT_VISIBILITY_END
+
+using Eigen::MatrixXf;
+using Eigen::Vector3f;
+
+class FHairCardGenCardSubdivider
+{
+private:
+	float SubdTolerance;
+	float SubdToleranceInitial;
+	const bool UseAdaptive;
+	const int MaxSubdivisions;
+	const int NumInterpPoints = 1000;
+	float LastDist;
+
+	MatrixXf MeanCurve;
+
+private:
+	static TArray<float> MatrixXfToTArray(const MatrixXf PointMatrix);
+	int NextSubdPoint(const int IndStart);
+	TArray<float> AdaptiveIter();
+	TArray<float> GetAdaptiveSubdPoints();
+	static MatrixXf LinearInterpolation(const MatrixXf, const int NumPoints);
+
+public:
+	FHairCardGenCardSubdivider(const float SubdTolerance, const bool UseAdaptive, const int MaxSubdivisions);
+	void SetSubdTolerance(const float NewSubdTolerance);
+	static MatrixXf TArrayToMatrixXf(const TArray<float> Points);
+	TArray<float> GetSubdivisionPoints(const MatrixXf Points);
+	static float GetCurveLength(const MatrixXf Points);
+	static float GetAverageCurvatureRadius(const MatrixXf Points);
+	static float GetToleranceFromLengthAndCurvRadius(const float Length, const float CurvRadius, const float Subdivisions);
+};
\ No newline at end of file
diff --git a/Engine/Plugins/Experimental/HairCardGenerator/Source/HairCardGeneratorEditor/Private/HairCardGenControllerBase.cpp b/Engine/Plugins/Experimental/HairCardGenerator/Source/HairCardGeneratorEditor/Private/HairCardGenControllerBase.cpp
new file mode 100644
--- /dev/null
+++ b/Engine/Plugins/Experimental/HairCardGenerator/Source/HairCardGeneratorEditor/Private/HairCardGenControllerBase.cpp
@@ -0,0 +1,125 @@
+// Copyright Epic Games, Inc. All Rights Reserved.
+
+#include "HairCardGenControllerBase.h"
+
+
+#include "GroomAsset.h"
+#include "MaterialDomain.h"
+#include "StaticMeshAttributes.h"
+#include "Engine/StaticMesh.h"
+#include "Materials/Material.h"
+
+#include "HairCardGeneratorPluginSettings.h"
+#include "HairCardGeneratorLog.h"
+
+#include UE_INLINE_GENERATED_CPP_BY_NAME(HairCardGenControllerBase)
+
+UHairCardGenControllerBase::UHairCardGenControllerBase()
+{
+}
+
+void UHairCardGenControllerBase::CreateCardsStaticMesh(UStaticMesh* StaticMesh, const TArray<float>& verts, const TArray<int32>& faces, const TArray<float>& normals, const TArray<float>& uvs, const TArray<int32>& groups)
+{	
+	// Set up a mapping so that each card gets its own polygon group id
+	int MaxGroup = 0;
+	for ( int GroupIdx : groups)
+		MaxGroup = FMath::Max(MaxGroup, GroupIdx);
+
+	// Set up basic model build settings
+	StaticMesh->SetNumSourceModels(1);
+	FStaticMeshSourceModel& Model = StaticMesh->GetSourceModel(0);
+	Model.BuildSettings.bRecomputeNormals = false;
+	Model.BuildSettings.bRecomputeTangents = true;
+	Model.BuildSettings.bComputeWeightedNormals = true;
+	Model.BuildSettings.bGenerateLightmapUVs = false;
+	Model.BuildSettings.bRemoveDegenerates = true;
+
+	FMeshDescription* MeshDescription = StaticMesh->CreateMeshDescription(0);
+	FStaticMeshAttributes MeshAttributes(*MeshDescription);
+
+	// Add a default material slot to the mesh
+	FName MaterialName = FName(TEXT("Material_0"));
+	FStaticMaterial DefaultMaterial;
+	DefaultMaterial.MaterialInterface = UMaterial::GetDefaultMaterial(EMaterialDomain::MD_Surface);
+	DefaultMaterial.MaterialSlotName = MaterialName;
+	DefaultMaterial.ImportedMaterialSlotName = MaterialName;
+
+	StaticMesh->GetStaticMaterials().Add(DefaultMaterial);
+
+	TVertexAttributesRef<FVector3f> VertexPositions = MeshAttributes.GetVertexPositions();
+	TPolygonGroupAttributesRef<FName> PolygonGroupMaterialSlotNames = MeshAttributes.GetPolygonGroupMaterialSlotNames();
+	TVertexInstanceAttributesRef<FVector2f> VertexInstanceUVs = MeshAttributes.GetVertexInstanceUVs();
+	TVertexInstanceAttributesRef<FVector3f> VertexInstanceNormals = MeshAttributes.GetVertexInstanceNormals();
+
+
+	// Add each unique group index to map and create a polygon group (also assign all groups default material)
+	TMap<int32,FPolygonGroupID> CardToGroupMap;
+	CardToGroupMap.Reserve(MaxGroup+1);
+	for ( int32 VertIndex = 0; VertIndex < groups.Num(); ++VertIndex )
+	{
+		// int GroupIndex = groups[VertIndex];
+		int GroupIndex = 0;
+		if ( CardToGroupMap.Contains(GroupIndex) )
+			continue;
+
+		const FPolygonGroupID PolygonGroupID = MeshDescription->CreatePolygonGroup();
+		PolygonGroupMaterialSlotNames[PolygonGroupID] = MaterialName;
+
+		CardToGroupMap.Add(GroupIndex, PolygonGroupID);
+	}
+
+	ensure(verts.Num() % 3 == 0);
+	for ( int VertIndex = 0; VertIndex < verts.Num()/3; ++VertIndex )
+	{
+		FVertexID VertexID = MeshDescription->CreateVertex();
+		VertexPositions[VertexID] = FVector3f(verts[3*VertIndex], verts[3*VertIndex+1], verts[3*VertIndex+2]);
+	}
+
+	for ( int TriIdx = 0; TriIdx < faces.Num() / 3; ++TriIdx )
+	{
+		TArray<FVertexInstanceID> TriVertIDs;
+		TriVertIDs.SetNum(3);
+
+		// Convert from counter-clockwise winding to clockwise (also flip precomputed normals)
+		const int32 VOrder[] = {2,1,0};
+		for ( int cidx = 0; cidx < 3; ++cidx )
+		{
+			const int32 TriCornerIdx = 3*TriIdx + VOrder[cidx];
+			const int32 VertCornerIdx = faces[TriCornerIdx];
+
+			const FVertexID VertexID(VertCornerIdx);
+			const FVertexInstanceID VertexInstanceID = MeshDescription->CreateVertexInstance(VertexID);
+
+			VertexInstanceUVs.Set(VertexInstanceID, 0, FVector2f(uvs[2*VertCornerIdx],uvs[2*VertCornerIdx+1]));
+			VertexInstanceNormals[VertexInstanceID] = -FVector3f(normals[3*VertCornerIdx],normals[3*VertCornerIdx+1],normals[3*VertCornerIdx+2]);
+
+			TriVertIDs[cidx] = VertexInstanceID;
+		}
+
+		// int GroupIndex = groups[faces[3*TriIdx]];
+		int GroupIndex = 0;
+		const FPolygonGroupID& PolygonGroupID = CardToGroupMap.FindRef(GroupIndex);
+
+		MeshDescription->CreatePolygon(PolygonGroupID, TriVertIDs);
+	}
+
+	StaticMesh->CommitMeshDescription(0);
+}
+
+TObjectPtr<UHairCardGeneratorPluginSettings>& UHairCardGenControllerBase::GetGroomSettings(TObjectPtr<UGroomAsset> Groom, int LODIndex, int GroupIndex)
+{
+	const FString SettingsTag = Groom->GetPathName() + TEXT(":LOD") + FString::FromInt(LODIndex);
+
+	// HACK: Could keep multiple settings objects around per-lod and group, but the group id should use the same settings so just update group id
+	TObjectPtr<UHairCardGeneratorPluginSettings>& Settings = GroomSettingsMap.FindOrAdd(SettingsTag);
+	if ( Settings)
+		Settings->SetGenerateForGroomGroup(GroupIndex);
+
+	return Settings;
+}
+
+void UHairCardGenControllerBase::UpdateGroomSettings(TObjectPtr<UGroomAsset> Groom, int LODIndex, int GroupID, TObjectPtr<UHairCardGeneratorPluginSettings> NewSettings)
+{
+	const FString SettingsTag = Groom->GetPathName() + TEXT(":LOD") + FString::FromInt(LODIndex);
+	GroomSettingsMap.FindOrAdd(SettingsTag) = NewSettings;
+}
diff --git a/Engine/Plugins/Experimental/HairCardGenerator/Source/HairCardGeneratorEditor/Private/HairCardGenControllerBase.h b/Engine/Plugins/Experimental/HairCardGenerator/Source/HairCardGeneratorEditor/Private/HairCardGenControllerBase.h
new file mode 100644
--- /dev/null
+++ b/Engine/Plugins/Experimental/HairCardGenerator/Source/HairCardGeneratorEditor/Private/HairCardGenControllerBase.h
@@ -0,0 +1,125 @@
+// Copyright Epic Games, Inc. All Rights Reserved.
+#pragma once
+
+#include "CoreMinimal.h"
+#include "UObject/Object.h"
+#include "GroomCacheData.h"
+#include "Containers/Map.h"
+#include "Templates/Tuple.h"
+
+
+#include "HairCardGenControllerBase.generated.h"
+
+class FHairDescription;
+//struct FHairCardGenerationData;
+
+USTRUCT(BlueprintType)
+struct FHairCardGen_StrandData
+{
+	GENERATED_BODY()
+
+	UPROPERTY(VisibleAnywhere, BlueprintReadOnly, Category = "Strand")
+	int GroupID = -1;
+};
+
+// With "Guide" data stipped out (since it isn't used as part of the mesh generation algorithm)
+USTRUCT(BlueprintType)
+struct FHairCardGen_GroomData
+{
+	GENERATED_BODY()
+
+	UPROPERTY(VisibleAnywhere, BlueprintReadOnly, Category = "Groom")
+	FString BasisType;
+
+	UPROPERTY(VisibleAnywhere, BlueprintReadOnly, Category = "Groom")
+	FString CurveType;
+
+	// Array of index buffers. Each buffer defines a distinct hair strand in the groom.
+	UPROPERTY(VisibleAnywhere, BlueprintReadOnly, Category = "Groom")
+	TArray<FHairCardGen_StrandData> Strands;
+
+	// Flattened Array of 3D vetex positions 
+	UPROPERTY(VisibleAnywhere, BlueprintReadOnly, Category = "Groom")
+	TArray<float> VertexPositions;
+
+	UPROPERTY(VisibleAnywhere, BlueprintReadOnly, Category = "Groom")
+	TArray<float> VertexWidths;
+
+	UPROPERTY(VisibleAnywhere, BlueprintReadOnly, Category = "Groom")
+	int32 HairlineGroupID = INDEX_NONE;
+};
+
+
+class UGroomAsset;
+class UHairCardGeneratorPluginSettings;
+
+/**
+ * 
+ */
+UCLASS(Blueprintable)
+class UHairCardGenControllerBase : public UObject
+{
+	GENERATED_BODY()
+
+public:
+	UHairCardGenControllerBase();
+
+	UFUNCTION(BlueprintImplementableEvent, Category = Python)
+	int GetPointsPerCurve();
+
+	UFUNCTION(BlueprintImplementableEvent, Category = Python)
+	bool LoadGroomData(const FHairCardGen_GroomData& GroomData, const FString& Name, const FString& CachedGroomsPath, const bool SaveCached = false);
+
+	UFUNCTION(BlueprintImplementableEvent, Category = Python)
+	bool LoadSettings(const UHairCardGeneratorPluginSettings* GeneratorSettings);
+
+	UFUNCTION(BlueprintImplementableEvent, Category = Python)
+	bool UpdateWidthMultipliers(const TArray<float>& HairWidths, const TArray<float>& RootScales, const TArray<float>& TipScales);
+
+	UFUNCTION(BlueprintImplementableEvent, Category = Python)
+	bool GenerateCardsGeometry();
+
+	UFUNCTION(BlueprintImplementableEvent, Category = Python)
+	bool GenerateClumps(const int GroupId = -1);
+
+	UFUNCTION(BlueprintImplementableEvent, Category = Python)
+	TArray<int> SetOptimizations(const int GroupId = -1);
+
+	UFUNCTION(BlueprintImplementableEvent, Category = Python)
+	TArray<float> GetAverageCurve(const int Id, const int Cid);
+
+	UFUNCTION(BlueprintImplementableEvent, Category = Python)
+	void SetInterpolatedAvgCurve(const int Id, const int Cid, const TArray<float>& Points);
+
+	UFUNCTION(BlueprintImplementableEvent, Category = Python)
+	bool ClusterTextures(int GroupIndex);
+
+	UFUNCTION(BlueprintImplementableEvent, Category = Python)
+	bool GenerateTextureLayout();
+
+	UFUNCTION(BlueprintImplementableEvent, Category = Python)
+	bool GenerateTextureAtlases(const float WidthScale = -1);
+
+	UFUNCTION(BlueprintImplementableEvent, Category = Python)
+	bool GenerateMesh(UStaticMesh* StaticMesh);
+
+	UFUNCTION(BlueprintImplementableEvent, Category = Python)
+	bool TestWriteFbxMesh();
+
+	UFUNCTION(BlueprintImplementableEvent, Category = Python)
+	bool CheckClusterTextures(int GroupIndex);
+
+	// Helper C++ function for building a static mesh from verts/faces, etc.
+	UFUNCTION(BlueprintCallable, Category = Python)
+	void CreateCardsStaticMesh(UStaticMesh* StaticMesh, const TArray<float>& verts, const TArray<int32>& faces, const TArray<float>& normals, const TArray<float>& uvs, const TArray<int32>& groups);
+
+
+	TObjectPtr<UHairCardGeneratorPluginSettings>& GetGroomSettings(TObjectPtr<UGroomAsset> Groom, int LODIndex, int GroupID);
+	void UpdateGroomSettings(TObjectPtr<UGroomAsset> Groom, int LODIndex, int GroupID, TObjectPtr<UHairCardGeneratorPluginSettings> NewSettings);
+
+private:
+
+
+	UPROPERTY(Transient)
+	TMap<FString,TObjectPtr<UHairCardGeneratorPluginSettings>> GroomSettingsMap;
+};
diff --git a/Engine/Plugins/Experimental/HairCardGenerator/Source/HairCardGeneratorEditor/Private/HairCardGenNaturalCubicSplines.cpp b/Engine/Plugins/Experimental/HairCardGenerator/Source/HairCardGeneratorEditor/Private/HairCardGenNaturalCubicSplines.cpp
new file mode 100644
--- /dev/null
+++ b/Engine/Plugins/Experimental/HairCardGenerator/Source/HairCardGeneratorEditor/Private/HairCardGenNaturalCubicSplines.cpp
@@ -0,0 +1,176 @@
+// Copyright Epic Games, Inc. All Rights Reserved.
+
+#include "HairCardGenNaturalCubicSplines.h"
+
+FHairCardGenNaturalCubicSplines::FHairCardGenNaturalCubicSplines(const TArray<float>& Knots, const TArray<float>& Values)
+{
+    MakeSpline(Knots, Values);
+    MakeBins();
+}
+
+float FHairCardGenNaturalCubicSplines::operator()(float t) const
+{
+    SplineData const& Spline = FindSpline(t);
+
+    double Value = Spline.Coefficients[3];
+    for (int i = 2; i >= 0; i--) {
+        Value *= t - Spline.Knot;
+        Value += Spline.Coefficients[i];
+    }
+
+    return Value;
+}
+
+void FHairCardGenNaturalCubicSplines::MakeSpline(const TArray<float>& Knots, const TArray<float>& Values)
+{
+    const int NumSegments = Knots.Num() - 1;
+
+    TArray<float> Intervals; Intervals.SetNum(NumSegments);
+    TArray<float> Slopes; Slopes.SetNum(NumSegments);
+
+    for (int i = 0; i < NumSegments; i++)
+    {
+        const float dt = Knots[i + 1] - Knots[i];
+        const float dx = Values[i + 1] - Values[i];
+
+        Intervals[i] = dt;
+        Slopes[i] = dx / dt;
+    }
+
+    TArray<float> L; L.SetNum(NumSegments + 1);
+    TArray<float> D; D.SetNum(NumSegments + 1);
+    TArray<float> U; U.SetNum(NumSegments + 1);
+    TArray<float> Y; Y.SetNum(NumSegments + 1);
+
+    D[0] = 1;
+    D[NumSegments] = 1;
+
+    for (int i = 1; i < NumSegments; i++)
+    {
+        L[i] = Intervals[i - 1];
+        D[i] = 2 * (Intervals[i - 1] + Intervals[i]);
+        U[i] = Intervals[i];
+        Y[i] = 6 * (Slopes[i] - Slopes[i - 1]);
+    }
+
+    SolveTridiagonalSystem(L, D, U, Y);
+    TArray<float> const& M = Y;
+
+    Splines.Empty();
+    Splines.SetNum(NumSegments);
+
+    for (int i = 0; i < NumSegments; i++) {
+        SplineData Spline;
+        Spline.Knot = Knots[i];
+        Spline.Coefficients[0] = Values[i];
+        Spline.Coefficients[1] = Slopes[i] - (M[i + 1] + 2 * M[i]) * Intervals[i] / 6;
+        Spline.Coefficients[2] = M[i] / 2;
+        Spline.Coefficients[3] = (M[i + 1] - M[i]) / (6 * Intervals[i]);
+        Splines[i] = Spline;
+    }
+
+    LowerBound = Knots[0];
+    UpperBound = Knots[NumSegments];
+}
+
+void FHairCardGenNaturalCubicSplines::MakeBins()
+{
+    BinInterval = (UpperBound - LowerBound) / (float)Splines.Num();
+
+    int Index = 0;
+
+    for (int Bin = 0; ; Bin++)
+    {
+        const float BinEdge = LowerBound + BinInterval * Bin;
+
+        while (Index + 1 < Splines.Num() && Splines[Index + 1].Knot <= BinEdge)
+        {
+            Index++;
+        }
+
+        Bins.Add(Index);
+
+        if (Index + 1 == Splines.Num())
+        {
+            break;
+        }
+    }
+
+    Bins.Shrink();
+}
+
+FHairCardGenNaturalCubicSplines::SplineData const& FHairCardGenNaturalCubicSplines::FindSpline(float t) const
+{
+    if (t <= LowerBound)
+    {
+        return Splines[0];
+    }
+    if (t >= UpperBound)
+    {
+        return Splines.Last();
+    }
+
+    int Bin = int((t - LowerBound) / BinInterval);
+    if (Bin >= Bins.Num())
+    {
+        Bin = Bins.Num() - 1;
+    }
+
+    int Index = Bins[Bin];
+    for (; Index + 1 < Splines.Num(); Index++)
+    {
+        if (t < Splines[Index + 1].Knot)
+        {
+            break;
+        }
+    }
+
+    return Splines[Index];
+}
+
+void FHairCardGenNaturalCubicSplines::SolveTridiagonalSystem(TArray<float>& Lower, TArray<float>& Diag, TArray<float>& Upper, TArray<float>& RHS)
+{
+    const int Dim = RHS.Num();
+
+    for(int i = 0; i < Dim - 1; i++)
+    {
+        if (abs(Diag[i]) >= abs(Lower[i + 1]))
+        {
+            const float W = Lower[i + 1] / Diag[i];
+            Diag[i + 1] -= W * Upper[i];
+            RHS[i + 1] -= W * RHS[i];
+            Lower[i + 1] = 0;
+        }
+        else
+        {
+            const float W = Diag[i] / Lower[i + 1];
+            const float U = Diag[i + 1];
+
+            Diag[i] = Lower[i + 1];
+            Diag[i + 1] = Upper[i] - W * U;
+            Lower[i + 1] = Upper[i + 1];
+            Upper[i + 1] *= -W;
+            Upper[i] = U;
+
+            const float R = RHS[i];
+            RHS[i] = RHS[i + 1];
+            RHS[i + 1] = R - W * RHS[i + 1];
+        }
+    }
+
+    RHS[Dim - 1] /= Diag[Dim - 1];
+
+    for (int iRev = 2; iRev <= Dim; iRev++)
+    {
+        int i = Dim - iRev;
+        if (i == Dim - 2)
+        {
+            RHS[i] -= Upper[i] * RHS[i + 1];
+        }
+        else
+        {
+            RHS[i] -= Upper[i] * RHS[i + 1] + Lower[i + 1] * RHS[i + 2];
+        }
+        RHS[i] /= Diag[i];
+    }
+}
\ No newline at end of file
diff --git a/Engine/Plugins/Experimental/HairCardGenerator/Source/HairCardGeneratorEditor/Private/HairCardGenNaturalCubicSplines.h b/Engine/Plugins/Experimental/HairCardGenerator/Source/HairCardGeneratorEditor/Private/HairCardGenNaturalCubicSplines.h
new file mode 100644
--- /dev/null
+++ b/Engine/Plugins/Experimental/HairCardGenerator/Source/HairCardGeneratorEditor/Private/HairCardGenNaturalCubicSplines.h
@@ -0,0 +1,31 @@
+// Copyright Epic Games, Inc. All Rights Reserved.
+
+#pragma once
+
+#include "Containers/Array.h"
+
+class FHairCardGenNaturalCubicSplines
+{
+private:
+    struct SplineData
+    {
+        double Knot;
+        double Coefficients[4];
+    };
+
+    TArray<SplineData> Splines;
+    TArray<int> Bins;
+    float LowerBound;
+    float UpperBound;
+    float BinInterval;
+
+private:
+    void MakeSpline(const TArray<float>& Knots, const TArray<float>& Values);
+    void MakeBins();
+    SplineData const& FindSpline(float t) const;
+    void SolveTridiagonalSystem(TArray<float>& Lower, TArray<float>& Diag, TArray<float>& Upper, TArray<float>& RHS);
+
+public:
+	FHairCardGenNaturalCubicSplines(const TArray<float>& Knots, const TArray<float>& Values);
+    float operator()(float t) const;
+};
\ No newline at end of file
diff --git a/Engine/Plugins/Experimental/HairCardGenerator/Source/HairCardGeneratorEditor/Private/HairCardGenOptionsWindow.cpp b/Engine/Plugins/Experimental/HairCardGenerator/Source/HairCardGeneratorEditor/Private/HairCardGenOptionsWindow.cpp
new file mode 100644
--- /dev/null
+++ b/Engine/Plugins/Experimental/HairCardGenerator/Source/HairCardGeneratorEditor/Private/HairCardGenOptionsWindow.cpp
@@ -0,0 +1,1021 @@
+// Copyright Epic Games, Inc. All Rights Reserved.
+
+#include "HairCardGenOptionsWindow.h"
+#include "HairCardGeneratorPluginSettings.h"
+
+#include "GroomAsset.h"
+#include "GroomAssetCards.h"
+#include "IDetailsView.h"
+#include "IDocumentation.h"
+#include "InputCoreTypes.h"
+#include "PropertyEditorModule.h"
+#include "SPrimaryButton.h"
+#include "HAL/FileManager.h"
+#include "HAL/PlatformApplicationMisc.h"
+#include "Input/Reply.h"
+#include "Interfaces/IMainFrameModule.h"
+#include "Misc/PackageName.h"
+#include "Styling/AppStyle.h"
+#include "Widgets/SWindow.h"
+#include "Framework/Application/SlateApplication.h"
+#include "Widgets/Images/SImage.h"
+#include "Widgets/Input/SButton.h"
+#include "Widgets/Input/SSpinBox.h"
+#include "Widgets/Input/SCheckBox.h"
+#include "Widgets/Layout/SBorder.h"
+#include "Widgets/Layout/SBox.h"
+#include "Widgets/Layout/SScrollBox.h"
+#include "Widgets/Layout/SUniformGridPanel.h"
+#include "Widgets/Text/STextBlock.h"
+#include "Widgets/Views/STableRow.h"
+#include "Widgets/Views/SListView.h"
+
+#include "HairCardGeneratorLog.h"
+
+#define LOCTEXT_NAMESPACE "HairCardSettings"
+
+/* SGroupSettingsListRow Declaration
+ *****************************************************************************/
+
+class SHairCardGenOptionsWindow;
+
+class SGroupSettingsListRow : public SMultiColumnTableRow< TObjectPtr<UHairCardGeneratorGroupSettings> >
+{
+public:
+	SLATE_BEGIN_ARGS( SGroupSettingsListRow ){}
+	SLATE_END_ARGS()
+public:
+	void Construct( const FArguments& InArgs, TObjectPtr<UHairCardGeneratorGroupSettings> InGroupSettings, const TSharedRef<STableViewBase>& InOwnerTableView );
+
+	virtual TSharedRef<SWidget> GenerateWidgetForColumn(const FName& ColumnName) override;
+
+private:
+	FText TestGetIndex() const;
+
+	// Card group filter widget
+	FText GetGroupFilterText() const;
+
+	// Number of cards widget
+	int32 ReadTargetNumCards() const;
+	void WriteTargetNumCards(int NumCards);
+
+	// Number of textures widget
+	int32 ReadTargetNumTextures() const;
+	void WriteTargetNumTextures(int NumTextures);
+
+	// Number of flyaway cards widget
+	int32 ReadMaxFlyaways() const;
+	void WriteMaxFlyaways(int MaxFlyaways);
+
+	// Number of target triangles widget
+	bool CheckEditableTriangles() const;
+	int32 ReadTargetNumTriangles() const;
+	void WriteTargetNumTriangles(int NumTris);
+
+	// Strand count widget
+	FText GetStrandCountText() const;
+
+	// Settings group generation state widget
+	FText GetGeneratedStateText() const;
+
+	// Remove Settings group (button widget)
+	FReply DeleteSettingsGroup();
+
+	// Return false if reduce from previous LOD is selected
+	bool CheckDisabledOnReducePrevLOD() const;
+
+	int32 GetIntPropertyMetadata(FName PropertyName, FName MetadataField, int32 DefaultValue = 0) const;
+
+private:
+	TObjectPtr<UHairCardGeneratorGroupSettings> GroupSettingsRow;
+};
+
+
+/* SHairCardGenOptionsWindow Declaration
+ *****************************************************************************/
+
+class SHairCardGenOptionsWindow : public SCompoundWidget
+{
+public:
+	SLATE_BEGIN_ARGS(SHairCardGenOptionsWindow)
+		: _SettingsObject(nullptr)
+		, _WidgetWindow()
+		, _MaxWindowHeight(0.0f)
+		, _MaxWindowWidth(0.0f)
+	{}
+
+		SLATE_ARGUMENT(TObjectPtr<UHairCardGeneratorPluginSettings>, SettingsObject)
+		SLATE_ARGUMENT(TSharedPtr<SWindow>, WidgetWindow)
+		SLATE_ARGUMENT(float, MaxWindowHeight)
+		SLATE_ARGUMENT(float, MaxWindowWidth)
+	SLATE_END_ARGS()
+public:
+	void Construct(const FArguments& InArgs);
+
+	//~ Begin SWidget interface
+	virtual bool SupportsKeyboardFocus() const override { return true; }
+	virtual FReply OnKeyDown(const FGeometry& MyGeometry, const FKeyEvent& InKeyEvent) override;
+	//~ End SWidget interface
+
+	SHairCardGenOptionsWindow()
+		: SettingsObject(nullptr)
+		, bRunOnClose(false)
+	{}
+
+	bool RunGenerationOnClose() const { return bRunOnClose; }
+
+private:
+	EActiveTimerReturnType SetFocusPostConstruct(double InCurrentTime, float InDeltaTime);
+	bool CanRunGeneration() const;
+	bool CanOpenAdvancedSettings() const;
+
+	FReply OnAddSettingsGroup();
+	FReply OnEditClick();
+	FReply OnGenerateClick();
+	FReply OnCancelClick();
+	FReply OnResetToDefaultClick() const;
+
+	FText GetTargetAssetName() const;
+
+	const FSlateBrush* GetErrorIcon() const;
+	EVisibility GetErrorIconVisibility() const;
+	FText GetErrorIconTooltip() const;
+	FText GetGenerateButtonTooltip() const;
+
+	void OnEditGroupSettings(TObjectPtr<UHairCardGeneratorGroupSettings> GroupSettings);
+
+	ESelectionMode::Type GetListSelectionMode() const { return ESelectionMode::Single; }
+	TSharedRef<ITableRow> OnGenerateOptionListRow(TObjectPtr<UHairCardGeneratorGroupSettings> InGroupSettings, const TSharedRef<STableViewBase>& OwnerTable);
+
+private:
+	// Main settings object
+	TObjectPtr<UHairCardGeneratorPluginSettings> SettingsObject;
+	// Should run (not cancelled)
+	bool bRunOnClose = false;
+
+	// Main window
+	TWeakPtr<SWindow> WidgetWindow;
+	// Generation settings details pane
+	TSharedPtr<class IDetailsView> DetailsView;
+	// List view for updating per settings-group options
+	TSharedPtr<SListView<TObjectPtr<UHairCardGeneratorGroupSettings>>> GroupSettingsListView;
+	TSharedPtr<SButton> RunButton;
+};
+
+
+/* SHairCardGroupSettingsWindow Declaration
+ *****************************************************************************/
+
+class SHairCardGroupSettingsWindow : public SCompoundWidget
+{
+public:
+	SLATE_BEGIN_ARGS(SHairCardGroupSettingsWindow)
+		: _FilterGroupIndex(0)
+		, _SettingsObject(nullptr)
+		, _WidgetWindow()
+		, _MaxWindowHeight(0.0f)
+		, _MaxWindowWidth(0.0f)
+	{}
+
+		SLATE_ARGUMENT(int, FilterGroupIndex)
+		SLATE_ARGUMENT(TObjectPtr<UHairCardGeneratorPluginSettings>, SettingsObject)
+		SLATE_ARGUMENT(TSharedPtr<SWindow>, WidgetWindow)
+		SLATE_ARGUMENT(float, MaxWindowHeight)
+		SLATE_ARGUMENT(float, MaxWindowWidth)
+	SLATE_END_ARGS()
+
+public:
+	void Construct(const FArguments& InArgs);
+
+	//~ Begin SWidget interface
+	virtual bool SupportsKeyboardFocus() const override { return true; }
+	virtual FReply OnKeyDown(const FGeometry& MyGeometry, const FKeyEvent& InKeyEvent) override;
+	//~ End SWidget interface
+
+	SHairCardGroupSettingsWindow()
+		: SettingsObject(nullptr)
+	{}
+
+private:
+	EActiveTimerReturnType SetFocusPostConstruct(double InCurrentTime, float InDeltaTime);
+	FReply OnOkClick();
+	FReply OnResetToDefaultClick() const;
+	FText GetIntermediateTargetName() const;
+
+private:
+	int FilterGroupIndex;
+	TObjectPtr<UHairCardGeneratorPluginSettings> SettingsObject;
+
+	// Main window
+	TWeakPtr<SWindow> WidgetWindow;
+	// Group settings detail view
+	TSharedPtr<IDetailsView> DetailsView;
+	TSharedPtr<SButton> OkButton;
+};
+
+
+/* HairCardGenWindow_Impl Declaration
+ *****************************************************************************/
+
+namespace HairCardGenWindow_Impl
+{
+	static bool IsValidPath(const FDirectoryPath& ContentPath, const FString& PkgName);
+	static bool DoFilesExist(const FDirectoryPath& ContentPath, const FString& PkgName);
+
+	static bool PromptGroupSettingsDialog(TSharedPtr<SWindow> ParentWindow, TObjectPtr<UHairCardGeneratorPluginSettings> Settings, int FilterGroupIndex);
+}
+
+
+/* HairCardGenWindow_Utils Definition
+ *****************************************************************************/
+
+bool HairCardGenWindow_Utils::PromptUserWithHairCardGenDialog(UHairCardGeneratorPluginSettings* SettingsObject)
+{
+	TSharedPtr<SWindow> ParentWindow;
+	if (FModuleManager::Get().IsModuleLoaded("MainFrame"))
+	{
+		IMainFrameModule& MainFrame = FModuleManager::LoadModuleChecked<IMainFrameModule>("MainFrame");
+		ParentWindow = MainFrame.GetParentWindow();
+	}
+
+	// Compute centered window position based on max window size, which include when all categories are expanded
+	const float WindowWidth = 600.0f;
+	const float WindowHeight = 400.0f;
+	FVector2D WindowSize = FVector2D(WindowWidth, WindowHeight); // Max window size it can get based on current slate
+
+	FSlateRect WorkAreaRect = FSlateApplicationBase::Get().GetPreferredWorkArea();
+	FVector2D DisplayTopLeft(WorkAreaRect.Left, WorkAreaRect.Top);
+	FVector2D DisplaySize(WorkAreaRect.Right - WorkAreaRect.Left, WorkAreaRect.Bottom - WorkAreaRect.Top);
+
+	float ScaleFactor = FPlatformApplicationMisc::GetDPIScaleFactorAtPoint(DisplayTopLeft.X, DisplayTopLeft.Y);
+	WindowSize *= ScaleFactor;
+
+	FVector2D WindowPosition = (DisplayTopLeft + (DisplaySize - WindowSize) / 2.0f) / ScaleFactor;
+
+	TSharedRef<SWindow> Window = SNew(SWindow)
+		.Title(LOCTEXT("HairCardSettings.Title", "Card Generation Settings"))
+		.SizingRule(ESizingRule::UserSized)
+		.AutoCenter(EAutoCenter::None)
+		.ClientSize(WindowSize)
+		.ScreenPosition(WindowPosition);
+
+	TSharedPtr<SHairCardGenOptionsWindow> HairCardGenWindow;
+	Window->SetContent
+	(
+		SAssignNew(HairCardGenWindow, SHairCardGenOptionsWindow)
+			.SettingsObject(SettingsObject)
+			.WidgetWindow(Window)
+	);
+
+	FSlateApplication::Get().AddModalWindow(Window, ParentWindow, /*bSlowTaskWindow =*/false);
+
+	return HairCardGenWindow->RunGenerationOnClose();
+}
+
+
+/* SGroupSettingsListRow Declaration
+ *****************************************************************************/
+void SGroupSettingsListRow::Construct( const FArguments& InArgs, TObjectPtr<UHairCardGeneratorGroupSettings> InGroupSettings, const TSharedRef<STableViewBase>& InOwnerTableView )
+{
+	GroupSettingsRow = InGroupSettings;
+	FSuperRowType::Construct( FSuperRowType::FArguments().Padding(0), InOwnerTableView);
+}
+
+TSharedRef<SWidget> SGroupSettingsListRow::GenerateWidgetForColumn(const FName& ColumnName)
+{
+	if ( ColumnName == TEXT("SettingsIndex") )
+	{
+		return SNew(STextBlock)
+			.Text(this, &SGroupSettingsListRow::TestGetIndex);
+	}
+	else if ( ColumnName == TEXT("CardGroupFilter") )
+	{
+		return SNew(STextBlock)
+			.Text(this, &SGroupSettingsListRow::GetGroupFilterText)
+			.ToolTipText(LOCTEXT("HairCardSettings.CardGroupFilter.ToolTip", "The card groups to which these settings apply"));
+	}
+	else if ( ColumnName == TEXT("TargetNumCards") )
+	{
+		return SNew(SSpinBox<int>)
+			.MinValue(GetIntPropertyMetadata(GET_MEMBER_NAME_CHECKED(UHairCardGeneratorGroupSettings, TargetNumberOfCards), TEXT("ClampMin"), 1))
+			.MaxValue(GetIntPropertyMetadata(GET_MEMBER_NAME_CHECKED(UHairCardGeneratorGroupSettings, TargetNumberOfCards), TEXT("ClampMax"), 10000))
+			.Value(this, &SGroupSettingsListRow::ReadTargetNumCards)
+			.OnValueChanged(this, &SGroupSettingsListRow::WriteTargetNumCards)
+			.IsEnabled(this, &SGroupSettingsListRow::CheckDisabledOnReducePrevLOD)
+			.ToolTipText(LOCTEXT("HairCardSettings.TargetNumCards.ToolTip", "The number of cards to generate"));
+	}
+	else if ( ColumnName == TEXT("TargetNumTextures") )
+	{
+		return SNew(SSpinBox<int>)
+			.MinValue(GetIntPropertyMetadata(GET_MEMBER_NAME_CHECKED(UHairCardGeneratorGroupSettings, NumberOfTexturesInAtlas), TEXT("ClampMin"), 1))
+			.MaxValue(GetIntPropertyMetadata(GET_MEMBER_NAME_CHECKED(UHairCardGeneratorGroupSettings, NumberOfTexturesInAtlas), TEXT("ClampMax"), 300))
+			.Value(this, &SGroupSettingsListRow::ReadTargetNumTextures)
+			.OnValueChanged(this, &SGroupSettingsListRow::WriteTargetNumTextures)
+			.IsEnabled(this, &SGroupSettingsListRow::CheckDisabledOnReducePrevLOD)
+			.ToolTipText(LOCTEXT("HairCardSettings.CardGroupFilter.ToolTip", "The number of textures to generate (must be <= number of cards)"));
+	}
+	else if ( ColumnName == TEXT("MaxFlyaways") )
+	{
+		return SNew(SSpinBox<int>)
+			.MinValue(GetIntPropertyMetadata(GET_MEMBER_NAME_CHECKED(UHairCardGeneratorGroupSettings, MaxNumberOfFlyaways), TEXT("ClampMin"), 0))
+			.MaxValue(GetIntPropertyMetadata(GET_MEMBER_NAME_CHECKED(UHairCardGeneratorGroupSettings, MaxNumberOfFlyaways), TEXT("ClampMax"), 1000))
+			.Value(this, &SGroupSettingsListRow::ReadMaxFlyaways)
+			.OnValueChanged(this, &SGroupSettingsListRow::WriteMaxFlyaways)
+			.ToolTipText(LOCTEXT("HairCardSettings.MaxFlyaways.ToolTip", "The maximum number of flyaway cards (single distinct strand cards) to allow"));
+	}
+	else if ( ColumnName == TEXT("TargetNumTriangles") )
+	{
+		return SNew(SSpinBox<int>)
+			.MinValue(GetIntPropertyMetadata(GET_MEMBER_NAME_CHECKED(UHairCardGeneratorGroupSettings, TargetTriangleCount), TEXT("ClampMin"), 2))
+			.MaxValue(GetIntPropertyMetadata(GET_MEMBER_NAME_CHECKED(UHairCardGeneratorGroupSettings, TargetTriangleCount), TEXT("ClampMax"), 100000))
+			.Value(this, &SGroupSettingsListRow::ReadTargetNumTriangles)
+			.OnValueChanged(this, &SGroupSettingsListRow::WriteTargetNumTriangles)
+			.IsEnabled(this, &SGroupSettingsListRow::CheckEditableTriangles)
+			.ToolTipText(LOCTEXT("HairCardSettings.CardGroupFilter.ToolTip", "The total number of triangles to use for all cards (approximate)"));
+	}
+	else if ( ColumnName == TEXT("NumStrands") )
+	{
+		return SNew(STextBlock)
+			.Text(this, &SGroupSettingsListRow::GetStrandCountText)
+			.ToolTipText(LOCTEXT("HairCardSettings.CardGroupFilter.ToolTip", "The number of strands these settings will apply to"));
+	}
+	else if ( ColumnName == TEXT("IsGenerated") )
+	{
+		return SNew(STextBlock)
+			.Text(this, &SGroupSettingsListRow::GetGeneratedStateText)
+			.ToolTipText(LOCTEXT("HairCardSettings.CardGroupFilter.ToolTip", "Indicates whether cards and textures have already been generated using these settings"));
+	}
+	else if ( ColumnName == TEXT("AddDeleteRow") )
+	{
+		return SNew(SButton)
+			.ButtonStyle( FAppStyle::Get(), "HoverHintOnly" )
+			.ToolTipText(LOCTEXT("HairCardSettings.RemoveGroup.ToolTip", "Remove this settings group"))
+			.OnClicked( this, &SGroupSettingsListRow::DeleteSettingsGroup )
+			// .ForegroundColor(FSlateColor::UseForeground())
+			.HAlign( HAlign_Center )
+			.VAlign( VAlign_Center )
+			.Content()
+			[
+				SNew(SImage)
+					.Image(FAppStyle::Get().GetBrush("Icons.Delete"))
+			];
+	}
+
+	return
+		SNew(STextBlock)
+		. Text( FText::Format( LOCTEXT("HairCardSettings.UnsupprtedColumn", "Unsupported Column: {0}"), FText::FromName( ColumnName ) ) );
+}
+
+FText SGroupSettingsListRow::TestGetIndex() const
+{
+	// NOTE: Indices can be wrong for a tick after Construct is called so currently just used as a label.
+	return FText::Format(LOCTEXT("HairCardSettings.SettingsIndex", "{0}"), FText::AsNumber(IndexInList));
+}
+
+FText SGroupSettingsListRow::GetGroupFilterText() const
+{
+	FStringBuilderBase ListBuilder;
+	ListBuilder.AppendChar(TCHAR('['));
+	for ( FName id : GroupSettingsRow->ApplyToCardGroups )
+	{
+		ListBuilder.Appendf(TEXT("%s, "), *id.ToString());
+	}
+	if ( GroupSettingsRow->ApplyToCardGroups.Num() > 0 )
+	{
+		ListBuilder.RemoveSuffix(2);
+	}
+	ListBuilder.AppendChar(TCHAR(']'));
+
+	return FText::Format(LOCTEXT("HairCardSettings.CardGroupFilter", "{0}"), FText::FromString(ListBuilder.ToString()));
+}
+
+int32 SGroupSettingsListRow::ReadTargetNumCards() const
+{
+	return GroupSettingsRow->TargetNumberOfCards;
+}
+
+void SGroupSettingsListRow::WriteTargetNumCards(int NumCards)
+{
+	GroupSettingsRow->TargetNumberOfCards = NumCards;
+}
+
+int32 SGroupSettingsListRow::ReadTargetNumTextures() const
+{
+	return GroupSettingsRow->NumberOfTexturesInAtlas;
+}
+
+void SGroupSettingsListRow::WriteTargetNumTextures(int NumTextures)
+{
+	GroupSettingsRow->NumberOfTexturesInAtlas = NumTextures;
+}
+
+int32 SGroupSettingsListRow::ReadMaxFlyaways() const
+{
+	return GroupSettingsRow->MaxNumberOfFlyaways;
+}
+
+void SGroupSettingsListRow::WriteMaxFlyaways(int MaxFlyaways)
+{
+	GroupSettingsRow->MaxNumberOfFlyaways = MaxFlyaways;
+}
+
+bool SGroupSettingsListRow::CheckEditableTriangles() const
+{
+	return GroupSettingsRow->UseAdaptiveSubdivision;
+}
+
+int32 SGroupSettingsListRow::ReadTargetNumTriangles() const
+{
+	return GroupSettingsRow->TargetTriangleCount;
+}
+
+void SGroupSettingsListRow::WriteTargetNumTriangles(int NumTris)
+{
+	GroupSettingsRow->TargetTriangleCount = NumTris;
+}
+
+FText SGroupSettingsListRow::GetStrandCountText() const
+{
+	return FText::Format(LOCTEXT("HairCardSettings.GroupStrandCount", "{0}"), FText::AsNumber(GroupSettingsRow->GetStrandCount()));
+}
+
+FText SGroupSettingsListRow::GetGeneratedStateText() const
+{
+	// TODO: Return specific generation string or split this field
+	uint8 NeedGenerateFlags = GroupSettingsRow->GetFilterGroupGeneratedFlags();
+	if ( NeedGenerateFlags == (uint8)EHairCardGenerationPipeline::None )
+		return FText::Format(LOCTEXT("HairCardSettings.GroupGenerate", "{0}"), FText::FromString(TEXT("No")));
+
+	if ( NeedGenerateFlags == (uint8)EHairCardGenerationPipeline::All )
+		return FText::Format(LOCTEXT("HairCardSettings.GroupGenerate", "{0}"), FText::FromString(TEXT("Yes")));
+
+	return  FText::Format(LOCTEXT("HairCardSettings.GroupGenerate", "{0}"), FText::FromString(TEXT("Partial")));
+}
+
+FReply SGroupSettingsListRow::DeleteSettingsGroup()
+{
+	GroupSettingsRow->RemoveSettingsFilterGroup();
+	return FReply::Handled();
+}
+
+bool SGroupSettingsListRow::CheckDisabledOnReducePrevLOD() const
+{
+	// Disable editing if reducing from previous LOD
+	return !GroupSettingsRow->GetReduceFromPreviousLOD();
+}
+
+int32 SGroupSettingsListRow::GetIntPropertyMetadata(FName PropertyName, FName MetadataField, int32 DefaultValue) const
+{
+	for (TFieldIterator<FProperty> PropertyIt(GroupSettingsRow->GetClass()); PropertyIt; ++PropertyIt)
+	{
+		if (PropertyIt->GetFName() != PropertyName)
+		{
+			continue;
+		}
+
+		if ( !PropertyIt->HasMetaData(MetadataField) )
+		{
+			return DefaultValue;
+		}
+
+		return  PropertyIt->GetIntMetaData(MetadataField);
+	}
+
+	return DefaultValue;
+}
+
+
+/* SHairCardGenOptionsWindow Definition
+ *****************************************************************************/
+
+void SHairCardGenOptionsWindow::Construct(const FArguments& InArgs)
+{
+	SettingsObject = InArgs._SettingsObject;
+	WidgetWindow = InArgs._WidgetWindow;
+
+	check(SettingsObject);
+
+	TSharedPtr<SBox> HeaderBox;
+	TSharedPtr<SBox> DetailsBox;
+	this->ChildSlot
+	[
+		SNew(SScrollBox)
+		+SScrollBox::Slot()
+		[
+			SNew(SVerticalBox)
+			+ SVerticalBox::Slot()
+				.AutoHeight()
+				.Padding(2)
+			[
+				SAssignNew(HeaderBox, SBox)
+			]
+			+ SVerticalBox::Slot()
+				.AutoHeight()
+				.Padding(2)
+			[
+				SAssignNew(DetailsBox, SBox)
+					.MaxDesiredHeight(500.0f)
+					.WidthOverride(700.0f)
+			]
+			+ SVerticalBox::Slot()
+				.AutoHeight()
+				.Padding(2)
+			[
+				SAssignNew(GroupSettingsListView, SListView<TObjectPtr<UHairCardGeneratorGroupSettings>>)
+					.ItemHeight(24)
+					.ListItemsSource(&SettingsObject->GetFilterGroupSettings())
+					.OnGenerateRow(this, &SHairCardGenOptionsWindow::OnGenerateOptionListRow)
+					.SelectionMode(this, &SHairCardGenOptionsWindow::GetListSelectionMode)
+					.OnMouseButtonDoubleClick(this, &SHairCardGenOptionsWindow::OnEditGroupSettings)
+					.HeaderRow(
+						SNew(SHeaderRow)
+							+ SHeaderRow::Column("SettingsIndex").DefaultLabel(LOCTEXT("HairCardSettings.SettingsIndex.Header", "Id")).FixedWidth(30)
+							+ SHeaderRow::Column("CardGroupFilter").DefaultLabel(LOCTEXT("HairCardSettings.CardGroupFilter.Header", "Card Groups"))
+							+ SHeaderRow::Column("TargetNumCards").DefaultLabel(LOCTEXT("HairCardSettings.TargetNumCards.Header", "# Cards"))
+							+ SHeaderRow::Column("TargetNumTextures").DefaultLabel(LOCTEXT("HairCardSettings.TargetNumTextures.Header", "# Textures"))
+							+ SHeaderRow::Column("TargetNumTriangles").DefaultLabel(LOCTEXT("HairCardSettings.TargetNumTriangles.Header", "# Triangles"))
+							+ SHeaderRow::Column("MaxFlyaways").DefaultLabel(LOCTEXT("HairCardSettings.MaxFlyaways.Header", "Max Flyaway Cards"))
+							+ SHeaderRow::Column("NumStrands").DefaultLabel(LOCTEXT("HairCardSettings.StrandCount.Header", "Strand Count"))
+							+ SHeaderRow::Column("IsGenerated").DefaultLabel(LOCTEXT("HairCardSettings.GroupGenerate.Header", "Generate"))
+							+ SHeaderRow::Column("AddDeleteRow").FixedWidth(50)
+							[
+								SNew(SButton)
+									.ButtonStyle( FAppStyle::Get(), "HoverHintOnly" )
+									.OnClicked(this, &SHairCardGenOptionsWindow::OnAddSettingsGroup)
+									.ToolTipText(LOCTEXT("HairCardSettings.AddSettingsGroup.ToolTip", "Add new settings group"))
+									.HAlign( HAlign_Center )
+									.VAlign( VAlign_Center )
+									[
+										SNew(SImage)
+											.Image(FAppStyle::Get().GetBrush("Icons.PlusCircle"))
+									]
+							]
+					)
+			]
+			+ SVerticalBox::Slot()
+				.AutoHeight()
+				.Padding(2)
+			[
+				SNew(SUniformGridPanel)
+					.SlotPadding(1)
+				+ SUniformGridPanel::Slot(0, 0)
+				[
+					SNew(SButton)
+						.HAlign(HAlign_Center)
+						.Text(LOCTEXT("HairCardSettings.GroupEditSettings", "Advanced..."))
+						.ToolTipText(LOCTEXT("HairCardSettings.GroupEditSettings.ToolTip", "Edit advanced settings for currently selected hair card group."))
+						.OnClicked(this, &SHairCardGenOptionsWindow::OnEditClick)
+						.IsEnabled(this, &SHairCardGenOptionsWindow::CanOpenAdvancedSettings)
+				]
+				+ SUniformGridPanel::Slot(2, 0)
+				[
+					SAssignNew(RunButton, SPrimaryButton)
+						.Text(LOCTEXT("HairCardSettings.Generate", "Generate"))
+						.ToolTipText(this, &SHairCardGenOptionsWindow::GetGenerateButtonTooltip)
+						.OnClicked(this, &SHairCardGenOptionsWindow::OnGenerateClick)
+						.IsEnabled(this, &SHairCardGenOptionsWindow::CanRunGeneration)
+				]
+				+ SUniformGridPanel::Slot(3, 0)
+				[
+					SNew(SButton)
+						.HAlign(HAlign_Center)
+						.Text(LOCTEXT("HairCardSettings.Cancel", "Cancel"))
+						.ToolTipText(LOCTEXT("HairCardSettings.Cancel.ToolTip", "Cancel hair card generation."))
+						.OnClicked(this, &SHairCardGenOptionsWindow::OnCancelClick)
+				]
+			]
+		]
+	];
+
+	FPropertyEditorModule& PropertyEditorModule = FModuleManager::GetModuleChecked<FPropertyEditorModule>("PropertyEditor");
+	FDetailsViewArgs DetailsViewArgs;
+	DetailsViewArgs.bAllowSearch = false;
+	DetailsViewArgs.NameAreaSettings = FDetailsViewArgs::HideNameArea;
+	DetailsView = PropertyEditorModule.CreateDetailView(DetailsViewArgs);
+
+	DetailsBox->SetContent(DetailsView->AsShared());
+
+	HeaderBox->SetContent(
+		SNew(SBorder)
+			.Padding(FMargin(3))
+			.BorderImage(FAppStyle::GetBrush("ToolPanel.GroupBorder"))
+		[
+			SNew(SHorizontalBox)
+			+SHorizontalBox::Slot()
+				.VAlign(VAlign_Center)
+				.AutoWidth()
+				.Padding(FMargin(2.0f, 0.0f))
+			[
+				SNew(SImage)
+					.Image(this, &SHairCardGenOptionsWindow::GetErrorIcon)
+					.ToolTipText(this, &SHairCardGenOptionsWindow::GetErrorIconTooltip)
+					.Visibility(this, &SHairCardGenOptionsWindow::GetErrorIconVisibility)
+			]
+			+ SHorizontalBox::Slot()
+				.AutoWidth()
+				.VAlign(VAlign_Center)
+			[
+				SNew(STextBlock)
+					.Text(this, &SHairCardGenOptionsWindow::GetTargetAssetName)
+			]
+			+ SHorizontalBox::Slot()
+				.HAlign(HAlign_Right)
+				.Padding(FMargin(2.0f, 0.0f))
+			[
+				SNew(SButton)
+					.Text(LOCTEXT("HairCardSettings.ResetOptions", "Reset hair card solver settings to defaults"))
+					.OnClicked(this, &SHairCardGenOptionsWindow::OnResetToDefaultClick)
+			]
+		]
+	);
+	DetailsView->SetObject(SettingsObject);
+	GroupSettingsListView->RequestListRefresh();
+
+	RegisterActiveTimer(0.f, FWidgetActiveTimerDelegate::CreateSP(this, &SHairCardGenOptionsWindow::SetFocusPostConstruct));
+}
+
+
+bool SHairCardGenOptionsWindow::CanRunGeneration() const
+{
+	TMap<uint32,uint32> ErrorCountMap;
+	SettingsObject->ValidateStrandFilters(ErrorCountMap);
+
+	if ( ErrorCountMap.Num() > 0 )
+		return false;
+
+	return HairCardGenWindow_Impl::IsValidPath(SettingsObject->DestinationPath, SettingsObject->BaseFilename);
+}
+
+bool SHairCardGenOptionsWindow::CanOpenAdvancedSettings() const
+{
+	int SelectCount = GroupSettingsListView->GetNumItemsSelected();
+	return !SettingsObject->bReduceCardsFromPreviousLOD && (SelectCount == 1);
+}
+
+FReply SHairCardGenOptionsWindow::OnKeyDown(const FGeometry& MyGeometry, const FKeyEvent& InKeyEvent)
+{
+	if (InKeyEvent.GetKey() == EKeys::Escape)
+	{
+		return OnCancelClick();
+	}
+	else if ( InKeyEvent.GetKey() == EKeys::Enter && CanRunGeneration() )
+	{
+		return OnGenerateClick();
+	}
+
+	return FReply::Unhandled();
+}
+
+EActiveTimerReturnType SHairCardGenOptionsWindow::SetFocusPostConstruct(double InCurrentTime, float InDeltaTime)
+{
+	if (RunButton.IsValid())
+	{
+		FSlateApplication::Get().SetKeyboardFocus(RunButton, EFocusCause::SetDirectly);
+	}
+
+	return EActiveTimerReturnType::Stop;
+}
+
+FReply SHairCardGenOptionsWindow::OnAddSettingsGroup()
+{
+	SettingsObject->AddNewSettingsFilterGroup();
+	return FReply::Handled();
+}
+
+FReply SHairCardGenOptionsWindow::OnEditClick()
+{
+	TArray<TObjectPtr<UHairCardGeneratorGroupSettings>> selectedItems;
+	int count = GroupSettingsListView->GetSelectedItems(selectedItems);
+	if ( count == 1 )
+		OnEditGroupSettings(selectedItems[0]);
+
+	return FReply::Handled();
+}
+
+FReply SHairCardGenOptionsWindow::OnCancelClick()
+{
+	if (WidgetWindow.IsValid())
+	{
+		WidgetWindow.Pin()->RequestDestroyWindow();
+	}
+	return FReply::Handled();
+}
+
+FReply SHairCardGenOptionsWindow::OnResetToDefaultClick() const
+{
+	SettingsObject->ResetToDefault();
+	DetailsView->SetObject(SettingsObject, true);
+	return FReply::Handled();
+}
+
+FReply SHairCardGenOptionsWindow::OnGenerateClick()
+{
+	bRunOnClose = true;
+	if (WidgetWindow.IsValid())
+	{
+		WidgetWindow.Pin()->RequestDestroyWindow();
+	}
+	return FReply::Handled();
+}
+
+void SHairCardGenOptionsWindow::OnEditGroupSettings(TObjectPtr<UHairCardGeneratorGroupSettings> FilterGroupSettings)
+{
+	if ( FilterGroupSettings == nullptr || !CanOpenAdvancedSettings() )
+		return;
+
+	int FilterGroupIndex = SettingsObject->GetFilterGroupSettings().Find(FilterGroupSettings);
+	if ( FilterGroupIndex == INDEX_NONE )
+		return;
+
+	HairCardGenWindow_Impl::PromptGroupSettingsDialog(WidgetWindow.Pin(), SettingsObject, FilterGroupIndex);
+}
+
+FText SHairCardGenOptionsWindow::GetTargetAssetName() const
+{
+	return FText::FromString(SettingsObject->BaseFilename);
+}
+
+const FSlateBrush* SHairCardGenOptionsWindow::GetErrorIcon() const
+{
+	TMap<uint32,uint32> ErrorCountMap;
+	SettingsObject->ValidateStrandFilters(ErrorCountMap);
+	if ( ErrorCountMap.Num() > 0 )
+		return FAppStyle::Get().GetBrush("Icons.ErrorWithColor");
+
+	if (!HairCardGenWindow_Impl::IsValidPath(SettingsObject->DestinationPath, SettingsObject->BaseFilename))
+		return FAppStyle::Get().GetBrush("Icons.ErrorWithColor");
+
+	return FAppStyle::Get().GetBrush("Icons.WarningWithColor");
+}
+
+EVisibility SHairCardGenOptionsWindow::GetErrorIconVisibility() const
+{
+	TMap<uint32,uint32> ErrorCountMap;
+	bool HasUnassignedStrands = SettingsObject->ValidateStrandFilters(ErrorCountMap);
+	if ( ErrorCountMap.Num() > 0 || HasUnassignedStrands )
+	{
+		return EVisibility::Visible;
+	}
+
+	if (!HairCardGenWindow_Impl::IsValidPath(SettingsObject->DestinationPath, SettingsObject->BaseFilename))
+	{
+		return EVisibility::Visible;
+	}
+
+	return HairCardGenWindow_Impl::DoFilesExist(SettingsObject->DestinationPath, SettingsObject->BaseFilename) ?
+		EVisibility::Visible : EVisibility::Collapsed;
+}
+
+FText SHairCardGenOptionsWindow::GetErrorIconTooltip() const
+{
+	TMap<uint32,uint32> ErrorCountMap;
+	bool HasUnassignedStrands = SettingsObject->ValidateStrandFilters(ErrorCountMap);
+	if ( ErrorCountMap.Num() > 0 )
+	{
+		return LOCTEXT("HairCardSettings.Generate.MultiAssignedStrandsError", "Some strands are included in more than one settings group, cannot run generation");
+	}
+	else if ( HasUnassignedStrands )
+	{
+		return LOCTEXT("HairCardSettings.Generate.UnassignedStrandsWarning", "The settings groups do not cover all groom strands, some strands will be left out of generation.");
+	}
+
+	if (!HairCardGenWindow_Impl::IsValidPath(SettingsObject->DestinationPath, SettingsObject->BaseFilename))
+	{
+		return LOCTEXT("HairCardSettings.Generate.InvalidFilenameWarning", "The specified filepath is not valid.");
+	}
+	return LOCTEXT("HairCardSettings.Generate.OverwriteFileWarning", "There are pre-existing files using this name. They are likely the previous card assets, which will be replaced and overwritten by this.");
+}
+
+FText SHairCardGenOptionsWindow::GetGenerateButtonTooltip() const
+{
+	TMap<uint32,uint32> ErrorCountMap;
+	bool HasUnassignedStrands = SettingsObject->ValidateStrandFilters(ErrorCountMap);
+	if ( ErrorCountMap.Num() > 0 )
+	{
+		return LOCTEXT("HairCardSettings.Generate.MultiAssignedStrandsError.ToolTip", "ERROR: Some strands are included in more than one settings group, cannot run generation");
+	}
+	else if ( HasUnassignedStrands )
+	{
+		return LOCTEXT("HairCardSettings.Generate.UnassignedStrandsWarning.ToolTip", "Warning: The settings groups do not cover all groom strands, some strands will be left out of generation");
+	}
+
+	return LOCTEXT("HairCardSettings.Generate.ToolTip", "Generate hair cards for all card groups. Will replace current card assets and textures on sucessful run.");
+}
+
+TSharedRef<ITableRow> SHairCardGenOptionsWindow::OnGenerateOptionListRow(TObjectPtr<UHairCardGeneratorGroupSettings> InGroupSettings, const TSharedRef<STableViewBase>& OwnerTable)
+{
+	return SNew(SGroupSettingsListRow, InGroupSettings, OwnerTable);
+}
+
+
+/* SHairCardGroupSettingsWindow Definition
+ *****************************************************************************/
+
+void SHairCardGroupSettingsWindow::Construct(const FArguments& InArgs)
+{
+	FilterGroupIndex = InArgs._FilterGroupIndex;
+	SettingsObject = InArgs._SettingsObject;
+	WidgetWindow = InArgs._WidgetWindow;
+
+	check(SettingsObject);
+	check(SettingsObject->GetFilterGroupSettings().Num() > FilterGroupIndex);
+
+	TSharedPtr<SBox> HeaderBox;
+	TSharedPtr<SBox> DetailsBox;
+	this->ChildSlot
+	[
+		SNew(SBox)
+			.MaxDesiredHeight(InArgs._MaxWindowHeight)
+			.MaxDesiredWidth(InArgs._MaxWindowWidth)
+		[
+			SNew(SVerticalBox)
+			+ SVerticalBox::Slot()
+				.AutoHeight()
+				.Padding(2)
+			[
+				SAssignNew(HeaderBox, SBox)
+			]
+			+ SVerticalBox::Slot()
+				.AutoHeight()
+				.Padding(2)
+			[
+				SAssignNew(DetailsBox, SBox)
+					.MaxDesiredHeight(700.0f)
+					.WidthOverride(400.0f)
+			]
+			+ SVerticalBox::Slot()
+				.AutoHeight()
+				.Padding(2)
+			[
+				SNew(SUniformGridPanel)
+					.SlotPadding(2)
+				+ SUniformGridPanel::Slot(1, 0)
+				[
+					SAssignNew(OkButton, SPrimaryButton)
+						.Text(LOCTEXT("HairCardSettings.Group.Ok", "Ok"))
+						.ToolTipText(LOCTEXT("HairCardSettings.Group.Ok.ToolTip", "Save hair card generation settings for this card group."))
+						.OnClicked(this, &SHairCardGroupSettingsWindow::OnOkClick)
+				]
+				// + SUniformGridPanel::Slot(2, 0)
+				// [
+				// 	SNew(SButton)
+				// 		.HAlign(HAlign_Center)
+				// 		.Text(LOCTEXT("HairCardSettings.Group.Cancel", "Cancel"))
+				// 		.ToolTipText(LOCTEXT("HairCardSettings.Group.Cancel.ToolTip", "Discards generation settings changes for this card group."))
+				// 		.OnClicked(this, &SHairCardGroupGenOptionsWindow::OnCancelClick)
+				// ]
+			]
+		]
+	];
+
+	FPropertyEditorModule& PropertyEditorModule = FModuleManager::GetModuleChecked<FPropertyEditorModule>("PropertyEditor");
+	FDetailsViewArgs DetailsViewArgs;
+	DetailsViewArgs.bAllowSearch = false;
+	DetailsViewArgs.NameAreaSettings = FDetailsViewArgs::HideNameArea;
+	DetailsView = PropertyEditorModule.CreateDetailView(DetailsViewArgs);
+
+	DetailsBox->SetContent(DetailsView->AsShared());
+
+	HeaderBox->SetContent(
+		SNew(SBorder)
+			.Padding(FMargin(3))
+			.BorderImage(FAppStyle::GetBrush("ToolPanel.GroupBorder"))
+		[
+			SNew(SHorizontalBox)
+			+ SHorizontalBox::Slot()
+				.AutoWidth()
+				.VAlign(VAlign_Center)
+			[
+				SNew(STextBlock)
+					.Text(this, &SHairCardGroupSettingsWindow::GetIntermediateTargetName)
+			]
+			+ SHorizontalBox::Slot()
+				.HAlign(HAlign_Right)
+				.Padding(FMargin(2.0f, 0.0f))
+			[
+				SNew(SButton)
+					.Text(LOCTEXT("HairCardGenWindow_ResetOptions", "Reset to Default"))
+					.OnClicked(this, &SHairCardGroupSettingsWindow::OnResetToDefaultClick)
+			]
+		]
+	);
+	DetailsView->SetObject(SettingsObject->GetFilterGroupSettings()[FilterGroupIndex]);
+
+	RegisterActiveTimer(0.f, FWidgetActiveTimerDelegate::CreateSP(this, &SHairCardGroupSettingsWindow::SetFocusPostConstruct));
+}
+
+FReply SHairCardGroupSettingsWindow::OnKeyDown(const FGeometry& MyGeometry, const FKeyEvent& InKeyEvent)
+{
+	// TODO: Support cancel/ok in details UI
+	if (InKeyEvent.GetKey() == EKeys::Escape || InKeyEvent.GetKey() == EKeys::Enter)
+	{
+		return OnOkClick();
+	}
+
+	return FReply::Unhandled();
+}
+
+FText SHairCardGroupSettingsWindow::GetIntermediateTargetName() const
+{
+	return FText::FromString(SettingsObject->BaseFilename + TEXT("_GIDX") + FString::FromInt(FilterGroupIndex));
+}
+
+EActiveTimerReturnType SHairCardGroupSettingsWindow::SetFocusPostConstruct(double InCurrentTime, float InDeltaTime)
+{
+	if (OkButton.IsValid())
+	{
+		FSlateApplication::Get().SetKeyboardFocus(OkButton, EFocusCause::SetDirectly);
+	}
+
+	return EActiveTimerReturnType::Stop;
+}
+
+FReply SHairCardGroupSettingsWindow::OnOkClick()
+{
+	if (WidgetWindow.IsValid())
+	{
+		WidgetWindow.Pin()->RequestDestroyWindow();
+	}
+
+	return FReply::Handled();
+}
+
+FReply SHairCardGroupSettingsWindow::OnResetToDefaultClick() const
+{
+	SettingsObject->ResetFilterGroupSettingsToDefault(FilterGroupIndex);
+	//Refresh the view to make sure the custom UI are updating correctly
+	DetailsView->SetObject(SettingsObject->GetFilterGroupSettings(FilterGroupIndex), true);
+
+	return FReply::Handled();
+}
+
+/* HairCardGenWindow_Impl Definition
+ *****************************************************************************/
+static bool HairCardGenWindow_Impl::IsValidPath(const FDirectoryPath& ContentPath, const FString& PkgName)
+{
+	FString TargetFilename;
+	bool bValidPath = FPackageName::TryConvertLongPackageNameToFilename(ContentPath.Path / PkgName, TargetFilename);
+
+	if (bValidPath)
+	{
+		bValidPath = FPaths::ValidatePath(TargetFilename);
+	}
+
+	return bValidPath;
+}
+
+static bool HairCardGenWindow_Impl::DoFilesExist(const FDirectoryPath & ContentPath, const FString & PkgName)
+{
+	FString TargetFilename;
+	if (!FPackageName::TryConvertLongPackageNameToFilename(ContentPath.Path / PkgName, TargetFilename))
+	{
+		return false;
+	}
+
+	TArray<FString> FoundFiles;
+	// Assume that wild card matches are the associate texture files
+	IFileManager::Get().FindFiles(FoundFiles, *(TargetFilename + TEXT("*.uasset")), /*Files =*/true, /*Directories =*/false);
+	return FoundFiles.Num() > 0;
+}
+
+static bool HairCardGenWindow_Impl::PromptGroupSettingsDialog(TSharedPtr<SWindow> ParentWindow, TObjectPtr<UHairCardGeneratorPluginSettings> Settings, int FilterGroupIndex)
+{
+	// Compute centered window position based on max window size, which include when all categories are expanded
+	const float WindowWidth = 450.0f;
+	const float WindowHeight = 750.0f;
+	FVector2D WindowSize = FVector2D(WindowWidth, WindowHeight); // Max window size it can get based on current slate
+
+	FSlateRect WorkAreaRect = FSlateApplicationBase::Get().GetPreferredWorkArea();
+	FVector2D DisplayTopLeft(WorkAreaRect.Left, WorkAreaRect.Top);
+	FVector2D DisplaySize(WorkAreaRect.Right - WorkAreaRect.Left, WorkAreaRect.Bottom - WorkAreaRect.Top);
+
+	float ScaleFactor = FPlatformApplicationMisc::GetDPIScaleFactorAtPoint(DisplayTopLeft.X, DisplayTopLeft.Y);
+	WindowSize *= ScaleFactor;
+
+	FVector2D WindowPosition = (DisplayTopLeft + (DisplaySize - WindowSize) / 2.0f) / ScaleFactor;
+
+	TSharedRef<SWindow> Window = SNew(SWindow)
+		.Title(LOCTEXT("HairCardSettings.Group.Title", "Card Group Settings"))
+		.SizingRule(ESizingRule::Autosized)
+		.AutoCenter(EAutoCenter::None)
+		.ClientSize(WindowSize)
+		.ScreenPosition(WindowPosition);
+
+	TSharedPtr<SHairCardGroupSettingsWindow> GroupSettingsOptionsWindow;
+	Window->SetContent
+	(
+		SAssignNew(GroupSettingsOptionsWindow, SHairCardGroupSettingsWindow)
+			.FilterGroupIndex(FilterGroupIndex)
+			.SettingsObject(Settings)
+			.WidgetWindow(Window)
+			.MaxWindowHeight(WindowHeight)
+			.MaxWindowWidth(WindowWidth)
+	);
+
+	FSlateApplication::Get().AddModalWindow(Window, ParentWindow, /*bSlowTaskWindow =*/false);
+
+	return true;
+}
+
+#undef LOCTEXT_NAMESPACE
+
diff --git a/Engine/Plugins/Experimental/HairCardGenerator/Source/HairCardGeneratorEditor/Private/HairCardGenOptionsWindow.h b/Engine/Plugins/Experimental/HairCardGenerator/Source/HairCardGeneratorEditor/Private/HairCardGenOptionsWindow.h
new file mode 100644
--- /dev/null
+++ b/Engine/Plugins/Experimental/HairCardGenerator/Source/HairCardGeneratorEditor/Private/HairCardGenOptionsWindow.h
@@ -0,0 +1,21 @@
+// Copyright Epic Games, Inc. All Rights Reserved.
+
+#pragma once
+
+#include "CoreMinimal.h"
+#include "GroomAssetCards.h" // for UHairCardGenerationSettings
+#include "HairCardGeneratorPluginSettings.h" // UHairCardGeneratorSettings
+
+class UGroomAsset;
+struct FHairGroupsCardsSourceDescription;
+
+namespace HairCardGenWindow_Utils
+{
+	/**
+	 * Opens a modal dialog, prompting the user to run hair card generation.
+	 * 
+	 * @param  OptionsObject	Settings structure defining the settings the user wants to run the generation with (may be modified on return)
+	 * @raturn True if the user specified they want to run the generation, false if they canceled the process.
+	 */
+	bool PromptUserWithHairCardGenDialog(UHairCardGeneratorPluginSettings* OptionsObject);
+}
diff --git a/Engine/Plugins/Experimental/HairCardGenerator/Source/HairCardGeneratorEditor/Private/HairCardGenSettingsDetail.cpp b/Engine/Plugins/Experimental/HairCardGenerator/Source/HairCardGeneratorEditor/Private/HairCardGenSettingsDetail.cpp
new file mode 100644
--- /dev/null
+++ b/Engine/Plugins/Experimental/HairCardGenerator/Source/HairCardGeneratorEditor/Private/HairCardGenSettingsDetail.cpp
@@ -0,0 +1,192 @@
+// Copyright Epic Games, Inc. All Rights Reserved.
+
+#include "HairCardGenSettingsDetail.h"
+#include "HairCardGeneratorPluginSettings.h"
+
+#include "DetailCategoryBuilder.h"
+#include "DetailLayoutBuilder.h"
+#include "DetailWidgetRow.h"
+#include "Widgets/Input/SCheckBox.h"
+#include "Widgets/Text/STextBlock.h"
+
+#include "HairCardGeneratorLog.h"
+
+#define LOCTEXT_NAMESPACE "HairCardSettingsDetails"
+
+void FHairCardSettingsDetailCustomization::CustomizeDetails(IDetailLayoutBuilder& DetailBuilder)
+{
+    TArray<TWeakObjectPtr<UHairCardGeneratorPluginSettings>> DetailObjects = DetailBuilder.GetObjectsOfTypeBeingCustomized<UHairCardGeneratorPluginSettings>();
+    if ( DetailObjects.Num() != 1 || !DetailObjects[0].IsValid() )
+    {
+        SettingsPtr = nullptr;    
+        return;
+    }
+
+    // Only support single object customized view (this is fine as HCS is a modal dialog anyway)
+    SettingsPtr = DetailObjects[0];
+
+    DetailBuilder.EditCategory(TEXT("Asset"));
+    DetailBuilder.EditCategory(TEXT("Import"));
+
+    IDetailCategoryBuilder& LodCategory = DetailBuilder.EditCategory(TEXT("Level Of Detail"));
+    const TSharedPtr<IPropertyHandle> ReduceFromLODHandle = DetailBuilder.GetProperty(GET_MEMBER_NAME_CHECKED(UHairCardGeneratorPluginSettings, bReduceCardsFromPreviousLOD));
+    const TSharedPtr<IPropertyHandle> UseReservedTxHandle = DetailBuilder.GetProperty(GET_MEMBER_NAME_CHECKED(UHairCardGeneratorPluginSettings, bUseReservedSpaceFromPreviousLOD));
+
+    DetailBuilder.EditCategory(TEXT("Randomness"));
+    IDetailCategoryBuilder& TxRenderCategory = DetailBuilder.EditCategory(TEXT("Texture Rendering"));
+
+    // Specialized widget with dynamic tooltip for enabled/disabled derive from previous LOD checkbox
+    DetailBuilder.HideProperty(ReduceFromLODHandle);
+    LodCategory.AddCustomRow(LOCTEXT("ReduceFromLOD.Filter", "Reduce Cards from Previous LOD"))
+    .NameContent()
+    [
+            SNew(STextBlock)
+            .Text(LOCTEXT("ReduceFromLOD.Name", "Reduce Cards from Previous LOD"))
+            .ToolTipText(ReduceFromLODHandle->GetToolTipText())
+            .Font(IDetailLayoutBuilder::GetDetailFont())
+    ]
+    .ValueContent()
+    [
+        SNew(SCheckBox)
+        .IsChecked(this, &FHairCardSettingsDetailCustomization::GetCheckValue, ReduceFromLODHandle)
+        .OnCheckStateChanged(this, &FHairCardSettingsDetailCustomization::SetCheckValue, ReduceFromLODHandle)
+        .IsEnabled(this, &FHairCardSettingsDetailCustomization::CheckReduceFromLOD)
+        .ToolTipText(this, &FHairCardSettingsDetailCustomization::ToolTipReduceFromLOD, ReduceFromLODHandle)
+    ];
+
+    // Specialized widget with dynamic tooltip for enabled/disabled use previous reserved texture space
+    DetailBuilder.HideProperty(UseReservedTxHandle);
+    TxRenderCategory.AddCustomRow(LOCTEXT("UseReservedTx.Filter", "Use Reserved Space from Previous LOD"))
+    .NameContent()
+    [
+            SNew(STextBlock)
+            .Text(LOCTEXT("UseReservedTx.Name", "Use Reserved Space from Previous LOD"))
+            .ToolTipText(UseReservedTxHandle->GetToolTipText())
+            .Font(IDetailLayoutBuilder::GetDetailFont())
+    ]
+    .ValueContent()
+    [
+        SNew(SCheckBox)
+        .IsChecked(this, &FHairCardSettingsDetailCustomization::GetCheckValue, UseReservedTxHandle)
+        .OnCheckStateChanged(this, &FHairCardSettingsDetailCustomization::SetCheckValue, UseReservedTxHandle)
+        .IsEnabled(this, &FHairCardSettingsDetailCustomization::CheckUseReservedTx)
+        .ToolTipText(this, &FHairCardSettingsDetailCustomization::ToolUseReservedTx, UseReservedTxHandle)
+    ];
+}
+
+
+TSharedRef<IDetailCustomization> FHairCardSettingsDetailCustomization::MakeInstance()
+{
+    return MakeShareable(new FHairCardSettingsDetailCustomization());
+}
+
+
+
+ECheckBoxState FHairCardSettingsDetailCustomization::GetCheckValue(const TSharedPtr<IPropertyHandle> Property) const
+{
+    bool CheckVal;
+    if ( Property.IsValid() && Property->GetValue(CheckVal) )
+    {
+        return (CheckVal) ? ECheckBoxState::Checked : ECheckBoxState::Unchecked;
+    }
+
+    return ECheckBoxState::Undetermined;
+}
+
+void FHairCardSettingsDetailCustomization::SetCheckValue(ECheckBoxState NewState, TSharedPtr<IPropertyHandle> Property) const
+{
+    if ( Property.IsValid() )
+    {
+        Property->SetValue(NewState == ECheckBoxState::Checked);
+    }
+}
+
+bool FHairCardSettingsDetailCustomization::CheckReduceFromLOD() const
+{
+    if ( !SettingsPtr.IsValid() )
+    {
+        return false;
+    }
+
+    if ( SettingsPtr.Get()->GetLODIndex() < 1 )
+    {
+        return false;
+    }
+    
+    if ( !SettingsPtr.Get()->HasValidFullParent() )
+    {
+        return false;;
+    }
+
+    return true;
+}
+
+FText FHairCardSettingsDetailCustomization::ToolTipReduceFromLOD(const TSharedPtr<IPropertyHandle> Property) const
+{
+    if ( !SettingsPtr.IsValid() )
+    {
+        return Property->GetToolTipText();
+    }
+
+    if ( SettingsPtr.Get()->GetLODIndex() < 1 )
+    {
+        return LOCTEXT("ReduceFromLOD.LOD0.ToolTip", "Cannot reduce LOD 0 (must run full generation)");
+    }
+
+    if ( !SettingsPtr.Get()->HasValidFullParent() )
+    {
+        return LOCTEXT("ReduceFromLOD.InvalidParent.ToolTip", "All lower LODs must be generated using the hair card generator tool");
+    }
+
+    return Property->GetToolTipText();
+}
+
+
+bool FHairCardSettingsDetailCustomization::CheckUseReservedTx() const
+{
+    if ( !SettingsPtr.IsValid() )
+    {
+        return false;
+    }
+
+    if ( SettingsPtr->bReduceCardsFromPreviousLOD )
+    {
+        return false;
+    }
+
+    if ( !SettingsPtr->HasDerivedTextureSettings() )
+    {
+        return false;
+    }
+
+    // TODO: Handle limiting/reserving texture by using texture resolution to compute reserved space in pixels
+    return (SettingsPtr->GetDerivedReservedTextureSize() >= 5);
+}
+
+FText FHairCardSettingsDetailCustomization::ToolUseReservedTx(const TSharedPtr<IPropertyHandle> Property) const
+{
+    if ( !SettingsPtr.IsValid() )
+    {
+        return Property->GetToolTipText();
+    }
+
+    if ( SettingsPtr->bReduceCardsFromPreviousLOD )
+    {
+        return LOCTEXT("UseReservedTx.Reducing.ToolTip", "Reduced card geometry will use previous LOD texture UVs");
+    }
+
+    if ( !SettingsPtr->HasDerivedTextureSettings() )
+    {
+        return LOCTEXT("UseReservedTx.InvalidParent.ToolTip", "Lower LODs must be generated using hair card generator tool with reserved space");
+    }
+
+    if ( SettingsPtr->GetDerivedReservedTextureSize() < 5 )
+    {
+        return LOCTEXT("UseReservedTx.NoReservedSpace.ToolTip", "No space reserved in parent LOD chain");
+    }
+
+    return Property->GetToolTipText();
+}
+
+
+#undef LOCTEXT_NAMESPACE
diff --git a/Engine/Plugins/Experimental/HairCardGenerator/Source/HairCardGeneratorEditor/Private/HairCardGenSettingsDetail.h b/Engine/Plugins/Experimental/HairCardGenerator/Source/HairCardGeneratorEditor/Private/HairCardGenSettingsDetail.h
new file mode 100644
--- /dev/null
+++ b/Engine/Plugins/Experimental/HairCardGenerator/Source/HairCardGeneratorEditor/Private/HairCardGenSettingsDetail.h
@@ -0,0 +1,36 @@
+// Copyright Epic Games, Inc. All Rights Reserved.
+
+#pragma once
+
+#include "CoreMinimal.h"
+#include "IDetailCustomization.h"
+#include "Styling/SlateTypes.h"
+
+//Forward declaration of IPropertyHandle.
+class IPropertyHandle;
+
+// Forward declater hair card settings object
+class UHairCardGeneratorPluginSettings;
+
+// Detail panel customization for HairCardGeneratorPluginSettings
+class FHairCardSettingsDetailCustomization : public IDetailCustomization
+{
+public:
+    //Function that customizes the Details Panel.
+    virtual void CustomizeDetails(IDetailLayoutBuilder& DetailBuilder) override;
+
+    //Returns a static instance of the Details Panel customization.
+    static TSharedRef<IDetailCustomization> MakeInstance();
+
+private:
+    ECheckBoxState GetCheckValue(const TSharedPtr<IPropertyHandle> Property) const;
+    void SetCheckValue(ECheckBoxState NewState, const TSharedPtr<IPropertyHandle> Property) const;
+
+    bool CheckReduceFromLOD() const;
+    FText ToolTipReduceFromLOD(const TSharedPtr<IPropertyHandle> Property) const;
+
+    bool CheckUseReservedTx() const;
+    FText ToolUseReservedTx(const TSharedPtr<IPropertyHandle> Property) const;
+
+    TWeakObjectPtr<UHairCardGeneratorPluginSettings> SettingsPtr;
+};
diff --git a/Engine/Plugins/Experimental/HairCardGenerator/Source/HairCardGeneratorEditor/Private/HairCardGenStrandFilter.cpp b/Engine/Plugins/Experimental/HairCardGenerator/Source/HairCardGeneratorEditor/Private/HairCardGenStrandFilter.cpp
new file mode 100644
--- /dev/null
+++ b/Engine/Plugins/Experimental/HairCardGenerator/Source/HairCardGeneratorEditor/Private/HairCardGenStrandFilter.cpp
@@ -0,0 +1,293 @@
+// Copyright Epic Games, Inc. All Rights Reserved.
+
+#include "HairCardGenStrandFilter.h"
+
+#include "HairCardGeneratorLog.h"
+
+
+/* FHairStrandAttribsRefProxy Implementation
+ *****************************************************************************/
+bool FHairStrandAttribsRefProxy::IsValid() const
+{
+    UE_CLOG(!StrandVertCounts.IsValid() || !VertexPositions.IsValid(), LogHairCardGenerator, Error, TEXT("Hair description has no vertex information!"));
+    UE_CLOG(VertexPositions.GetNumElements() == 0, LogHairCardGenerator, Error, TEXT("Hair description has no vertices!"))
+
+    return StrandVertCounts.IsValid() && VertexPositions.IsValid() && (VertexPositions.GetNumElements() != 0);
+}
+
+bool FHairStrandAttribsRefProxy::IsGuideStrand() const
+{
+    return GuideStrands.IsValid() && (GuideStrands[FStrandID(StrandIndex)] != 0);
+}
+
+bool FHairStrandAttribsRefProxy::HasGroupID() const
+{
+    return StrandGroupIds.IsValid();
+}
+
+const int& FHairStrandAttribsRefProxy::GetGroupID() const
+{
+    return StrandGroupIds[FStrandID(StrandIndex)];
+}
+
+bool FHairStrandAttribsRefProxy::HasStrandWidth() const
+{
+    return StrandWidths.IsValid();
+}
+
+const float& FHairStrandAttribsRefProxy::GetStrandWidth() const
+{
+    return StrandWidths[FStrandID(StrandIndex)];
+}
+
+bool FHairStrandAttribsRefProxy::HasCardGroupName() const
+{
+    return StrandGroupCardsNames.IsValid();
+}
+
+const FName& FHairStrandAttribsRefProxy::GetCardGroupName() const
+{
+    return StrandGroupCardsNames[FStrandID(StrandIndex)];
+}
+
+bool FHairStrandAttribsRefProxy::HasBasisType() const
+{
+    return StrandBasisType.IsValid();
+}
+
+EGroomBasisType FHairStrandAttribsRefProxy::GetBasisType() const
+{
+    UEnum* BasisEnum = StaticEnum<EGroomBasisType>();
+    return (StrandBasisType.IsValid() ? (EGroomBasisType)BasisEnum->GetValueByName(StrandBasisType[FStrandID(StrandIndex)]) : EGroomBasisType::NoBasis);
+}
+
+bool FHairStrandAttribsRefProxy::HasCurveType() const
+{
+    return StrandCurveType.IsValid();
+}
+
+EGroomCurveType FHairStrandAttribsRefProxy::GetCurveType() const
+{
+    UEnum* CurveEnum = StaticEnum<EGroomCurveType>();
+    return (StrandCurveType.IsValid() ? (EGroomCurveType)CurveEnum->GetValueByName(StrandCurveType[FStrandID(StrandIndex)]) : EGroomCurveType::Linear);
+}
+
+bool FHairStrandAttribsRefProxy::HasKnots() const
+{
+    return StrandKnots.IsValid();
+}
+
+TArrayAttribute<const float> FHairStrandAttribsRefProxy::GetKnots() const
+{
+    return StrandKnots[FStrandID(StrandIndex)];
+}
+
+int FHairStrandAttribsRefProxy::NumVerts() const
+{
+    return StrandVertCounts[FStrandID(StrandIndex)];
+}
+
+const FVector3f& FHairStrandAttribsRefProxy::GetVertexPosition(int Index) const
+{
+    const FVertexID VertID(StartVertIndex + Index);
+    return VertexPositions[VertID];
+}
+
+TArrayView<const FVector3f> FHairStrandAttribsRefProxy::GetVertexPositions() const
+{
+    const FStrandID StrandID(StrandIndex);
+    int VertCount = StrandVertCounts[StrandID];
+
+    return TArrayView<const FVector3f>(VertexPositions.GetArrayView(StartVertIndex).GetData(), VertCount);
+}
+
+bool FHairStrandAttribsRefProxy::HasVertexWidths() const
+{
+    return VertexWidths.IsValid();
+}
+
+const float& FHairStrandAttribsRefProxy::GetVertexWidth(int Index) const
+{
+    const FVertexID VertID(StartVertIndex + Index);
+    return VertexWidths[VertID];
+}
+
+TArrayView<const float> FHairStrandAttribsRefProxy::GetVertexWidths() const
+{
+    const FStrandID StrandID(StrandIndex);
+    int VertCount = StrandVertCounts[StrandID];
+
+    return TArrayView<const float>(VertexWidths.GetArrayView(StartVertIndex).GetData(), VertCount);
+}
+
+
+FHairStrandAttribsRefProxy::FHairStrandAttribsRefProxy(const FHairDescription& InHairDescription)
+    : StrandIndex(0)
+    , StartVertIndex(0)
+    , GuideStrands(InHairDescription.StrandAttributes().GetAttributesRef<int>(HairAttribute::Strand::Guide))
+    , StrandGroupIds(InHairDescription.StrandAttributes().GetAttributesRef<int>(HairAttribute::Strand::GroupID))
+    , StrandWidths(InHairDescription.StrandAttributes().GetAttributesRef<float>(HairAttribute::Strand::Width))
+    , StrandGroupCardsNames(InHairDescription.StrandAttributes().GetAttributesRef<FName>(HairAttribute::Strand::GroupCardsName))
+    , StrandBasisType(InHairDescription.StrandAttributes().GetAttributesRef<FName>(HairAttribute::Strand::BasisType))
+    , StrandCurveType(InHairDescription.StrandAttributes().GetAttributesRef<FName>(HairAttribute::Strand::CurveType))
+    , StrandKnots(InHairDescription.StrandAttributes().GetAttributesRef<TArrayAttribute<float>>(HairAttribute::Strand::Knots))
+    , StrandVertCounts(InHairDescription.StrandAttributes().GetAttributesRef<int>(HairAttribute::Strand::VertexCount))
+    , VertexPositions(InHairDescription.VertexAttributes().GetAttributesRef<FVector3f>(HairAttribute::Vertex::Position))
+    , VertexWidths(InHairDescription.VertexAttributes().GetAttributesRef<float>(HairAttribute::Vertex::Width))
+{}
+
+
+void FHairStrandAttribsRefProxy::SetStrand(int InStrandIndex)
+{
+    int StartIdx = FMath::Min(InStrandIndex,StrandIndex);
+    int EndIdx = FMath::Max(InStrandIndex,StrandIndex);
+
+    // Accumulate vertex counts (for increment op will just add StrandVertCounts[StrandIndex])
+    int dir = FMath::Sign(EndIdx-StartIdx);
+    for ( int i=StartIdx; i < EndIdx; ++i )
+        StartVertIndex += dir * StrandVertCounts[FStrandID(i)];
+
+    StrandIndex = InStrandIndex;
+}
+
+/* FHairStrandAttribsIteratorEnd Implementation
+ *****************************************************************************/
+FHairStrandAttribsIteratorEnd::FHairStrandAttribsIteratorEnd(const FHairDescription& InHairDescription, int32 InEndIndex /* = INDEX_NONE */)
+    : EndIndex((InEndIndex == INDEX_NONE) ? InHairDescription.GetNumStrands() : InEndIndex)
+    , HairDescription(InHairDescription)
+{}
+
+
+/* FHairStrandAttribsConstIterator Implementation
+ *****************************************************************************/
+FHairStrandAttribsConstIterator::FHairStrandAttribsConstIterator(const FHairDescription& InHairDescription, int32 InStartIndex /* = 0 */, FHairStrandAttribsConstIterator::EIgnoreStrandFlags InIgnoreFlags)
+    : StrandIndex(InStartIndex)
+    , IgnoreFlags(InIgnoreFlags)
+    , HairDescription(InHairDescription)
+    , AttribsProxy(HairDescription)
+{
+    if ( StrandIndex < HairDescription.GetNumStrands() )
+        FindFiltered();
+}
+
+const FHairStrandAttribsRefProxy& FHairStrandAttribsConstIterator::operator*() const
+{
+    return AttribsProxy;
+}
+
+const FHairStrandAttribsRefProxy* FHairStrandAttribsConstIterator::operator->() const
+{
+    return &AttribsProxy;
+}
+
+FHairStrandAttribsConstIterator& FHairStrandAttribsConstIterator::operator++()
+{
+    if ( StrandIndex >= HairDescription.GetNumStrands() )
+    {
+        return *this;
+    }
+
+    StrandIndex += 1;
+    FindFiltered();
+
+    return *this;
+}
+
+bool FHairStrandAttribsConstIterator::operator!=(const FHairStrandAttribsConstIterator& Iter) const
+{
+    UE_CLOG((&Iter.HairDescription != &HairDescription), LogHairCardGenerator, Error, TEXT("Invalid strand iterator comparison (different hair descriptions)!"));
+    return (StrandIndex != Iter.StrandIndex);
+}
+
+bool FHairStrandAttribsConstIterator::operator!=(const FHairStrandAttribsIteratorEnd& Iter) const
+{
+    return (StrandIndex < Iter.EndIndex);
+}
+
+void FHairStrandAttribsConstIterator::FindFiltered()
+{
+    for ( ; StrandIndex < HairDescription.GetNumStrands(); ++StrandIndex )
+    {
+        AttribsProxy.SetStrand(StrandIndex);
+        if ( (IgnoreFlags & IgnoreGuides) && AttribsProxy.IsGuideStrand() )
+            continue;
+
+        if ( (IgnoreFlags & IgnoreEmpty) && (AttribsProxy.NumVerts() == 0) )
+            continue;
+
+        return;
+    }
+}
+
+
+/* FStrandAttributesRange Implementation
+ *****************************************************************************/
+
+FStrandAttributesRange::FStrandAttributesRange(const FHairDescription& HairDesc, FHairStrandAttribsConstIterator::EIgnoreStrandFlags InIgnoreFlags/* = FHairStrandAttribsConstIterator::DefaultIgnore*/)
+    : HairDescription(HairDesc)
+    , BeginIter(HairDesc, 0, InIgnoreFlags)
+{}
+
+bool FStrandAttributesRange::IsValid() const
+{
+    return BeginIter->IsValid();
+}
+
+int FStrandAttributesRange::NumStrands() const
+{
+    return HairDescription.GetNumStrands();
+}
+
+int FStrandAttributesRange::NumVerts() const
+{
+    return BeginIter->IsValid() ? BeginIter->VertexPositions.GetNumElements() : 0;
+}
+
+const FHairStrandAttribsRefProxy& FStrandAttributesRange::First() const
+{
+    return *BeginIter;
+}
+
+FHairStrandAttribsConstIterator FStrandAttributesRange::begin() const
+{
+    return BeginIter;
+}
+
+FHairStrandAttribsIteratorEnd FStrandAttributesRange::end() const
+{
+    return FHairStrandAttribsIteratorEnd(HairDescription);
+}
+
+/* FStrandFilterOp_And Implementation
+ *****************************************************************************/
+FStrandFilterOp_And::FStrandFilterOp_And(TUniquePtr<FComposableStrandFilterOp> InLeftOp, TUniquePtr<FComposableStrandFilterOp> InRightOp)
+    : LeftOp(MoveTemp(InLeftOp)), RightOp(MoveTemp(InRightOp))
+{}
+
+bool FStrandFilterOp_And::FilterStrand(const FHairStrandAttribsRefProxy& StrandAttributes)
+{
+    return (LeftOp->FilterStrand(StrandAttributes) && RightOp->FilterStrand(StrandAttributes));
+}
+
+/* FStrandFilterOp_Or Implementation
+ *****************************************************************************/
+FStrandFilterOp_Or::FStrandFilterOp_Or(TUniquePtr<FComposableStrandFilterOp> InLeftOp, TUniquePtr<FComposableStrandFilterOp> InRightOp)
+    : LeftOp(MoveTemp(InLeftOp)), RightOp(MoveTemp(InRightOp))
+{}
+
+bool FStrandFilterOp_Or::FilterStrand(const FHairStrandAttribsRefProxy& StrandAttributes)
+{
+    return (LeftOp->FilterStrand(StrandAttributes) || RightOp->FilterStrand(StrandAttributes));
+}
+
+/* FStrandFilterOp_InCardGroup Implementation
+ *****************************************************************************/
+FStrandFilterOp_InCardGroup::FStrandFilterOp_InCardGroup(const TSet<FName>& InCardGroupIds)
+    : CardGroupNames(InCardGroupIds)
+{}
+
+bool FStrandFilterOp_InCardGroup::FilterStrand(const FHairStrandAttribsRefProxy& StrandAttributes)
+{
+    // Return true for all strands if no valid card group id, otherwise check if group is in set
+    return StrandAttributes.HasCardGroupName() ? CardGroupNames.Contains(StrandAttributes.GetCardGroupName()) : true;
+}
diff --git a/Engine/Plugins/Experimental/HairCardGenerator/Source/HairCardGeneratorEditor/Private/HairCardGenStrandFilter.h b/Engine/Plugins/Experimental/HairCardGenerator/Source/HairCardGeneratorEditor/Private/HairCardGenStrandFilter.h
new file mode 100644
--- /dev/null
+++ b/Engine/Plugins/Experimental/HairCardGenerator/Source/HairCardGeneratorEditor/Private/HairCardGenStrandFilter.h
@@ -0,0 +1,190 @@
+// Copyright Epic Games, Inc. All Rights Reserved.
+
+#pragma once
+
+#include "CoreMinimal.h"
+#include "GroomAsset.h"
+#include "GroomCacheData.h"
+
+#include "Misc/Optional.h"
+
+
+// Reference structure used to proxy a particular strand's attributes
+class FStrandAttributesRange;
+class FHairStrandAttribsConstIterator;
+class FHairStrandAttribsRefProxy
+{
+    friend class FStrandAttributesRange;
+    friend class FHairStrandAttribsConstIterator;
+public:
+    bool IsValid() const;
+
+    bool IsGuideStrand() const;
+
+    bool HasGroupID() const;
+    const int& GetGroupID() const;
+
+    bool HasStrandWidth() const;
+    const float& GetStrandWidth() const;
+
+    bool HasCardGroupName() const;
+    const FName& GetCardGroupName() const;
+
+    bool HasBasisType() const;
+    EGroomBasisType GetBasisType() const;
+
+    bool HasCurveType() const;
+    EGroomCurveType GetCurveType() const;
+
+    bool HasKnots() const;
+    TArrayAttribute<const float> GetKnots() const;
+
+    // Per-strand vertex info
+    int NumVerts() const;
+    const FVector3f& GetVertexPosition(int Index) const;
+    TArrayView<const FVector3f> GetVertexPositions() const;
+
+    bool HasVertexWidths() const;
+    const float& GetVertexWidth(int Index) const;
+    TArrayView<const float> GetVertexWidths() const;
+
+
+private:
+    FHairStrandAttribsRefProxy() = delete;
+    FHairStrandAttribsRefProxy(const FHairDescription& InHairDescription);
+
+    void SetStrand(int InStrandIndex);
+
+    int StrandIndex;
+    int StartVertIndex;
+
+    TStrandAttributesConstRef<int> GuideStrands;
+    TStrandAttributesConstRef<int> StrandGroupIds;
+    TStrandAttributesConstRef<float> StrandWidths;
+    TStrandAttributesConstRef<FName> StrandGroupCardsNames;
+
+    TStrandAttributesConstRef<FName> StrandBasisType;
+    TStrandAttributesConstRef<FName> StrandCurveType;
+    TStrandAttributesConstRef<TArrayAttribute<float>> StrandKnots;
+
+    TStrandAttributesConstRef<int> StrandVertCounts;
+    TVertexAttributesConstRef<FVector3f> VertexPositions;
+    TVertexAttributesConstRef<float> VertexWidths;
+};
+
+
+// Helper iterator clases to make it easier to iterate through strands and get access to per-strand attributes
+class FHairStrandAttribsIteratorEnd
+{
+    friend class FHairStrandAttribsConstIterator;
+public:
+    FHairStrandAttribsIteratorEnd(const FHairDescription& InHairDescription, int32 InEndIndex = INDEX_NONE);
+
+private:
+    int32 EndIndex;
+    const FHairDescription& HairDescription;
+};
+
+class FHairStrandAttribsConstIterator
+{
+public:
+    enum EIgnoreStrandFlags
+    {
+        // Don't ignore anything
+        IgnoreNone = 0x0,
+
+        // Ignore guide strands in traversal
+        IgnoreGuides = 0x01,
+        // Ignore empty (NumVerts==0) strands in traversal
+        IgnoreEmpty = 0x02,
+
+        // Default setting is to ignore guide and empty strands
+        DefaultIgnore = IgnoreGuides | IgnoreEmpty,
+    };
+
+public:
+    FHairStrandAttribsConstIterator(const FHairDescription& InHairDescription, int32 InStartIndex = 0, EIgnoreStrandFlags InIgnoreFlags = EIgnoreStrandFlags::DefaultIgnore);
+
+    const FHairStrandAttribsRefProxy& operator*() const;
+    const FHairStrandAttribsRefProxy* operator->() const;
+    FHairStrandAttribsConstIterator& operator++();
+    bool operator!=(const FHairStrandAttribsConstIterator& Iter) const;
+    bool operator!=(const FHairStrandAttribsIteratorEnd& Iter) const;
+
+private:
+    FHairStrandAttribsConstIterator() = delete;
+
+    void FindFiltered();
+
+    int32 StrandIndex;
+    EIgnoreStrandFlags IgnoreFlags;
+    const FHairDescription& HairDescription;
+
+    FHairStrandAttribsRefProxy AttribsProxy;
+};
+
+// Helper traversal class allows use of range-for syntax for getting per-strand attributes
+class FStrandAttributesRange
+{
+public:
+    /// Creates a helper class which will return iterator over each strand in a hair description with an attribute proxy to make it relatively easy to access per-strand attributes
+    /// @param HairDesc FHairDescription for the groom
+    /// @param InIgnoreFlags Flags indicating what strands to ignore during traversal (defaults to ignore guide strands and empty (vert) strands)
+	FStrandAttributesRange(const FHairDescription& HairDesc, FHairStrandAttribsConstIterator::EIgnoreStrandFlags InIgnoreFlags = FHairStrandAttribsConstIterator::DefaultIgnore);
+
+    bool IsValid() const;
+    int NumStrands() const;
+    int NumVerts() const;
+
+    const FHairStrandAttribsRefProxy& First() const;
+
+    FHairStrandAttribsConstIterator begin() const;
+    FHairStrandAttribsIteratorEnd end() const;
+private:
+    const FHairDescription& HairDescription;
+    const FHairStrandAttribsConstIterator BeginIter;
+};
+
+
+// Base class for all strand filters
+class FComposableStrandFilterOp
+{
+public:
+	virtual bool FilterStrand(const FHairStrandAttribsRefProxy& StrandAttributes) = 0;
+
+    virtual ~FComposableStrandFilterOp() = default;
+protected:
+	FComposableStrandFilterOp() = default;
+};
+
+class FStrandFilterOp_And : public FComposableStrandFilterOp
+{
+public:
+    FStrandFilterOp_And(TUniquePtr<FComposableStrandFilterOp> InLeftOp, TUniquePtr<FComposableStrandFilterOp> InRightOp);
+    virtual bool FilterStrand(const FHairStrandAttribsRefProxy& StrandAttributes) override;
+    
+private:
+    TUniquePtr<FComposableStrandFilterOp> LeftOp;
+    TUniquePtr<FComposableStrandFilterOp> RightOp;
+};
+
+class FStrandFilterOp_Or : public FComposableStrandFilterOp
+{
+public:
+    FStrandFilterOp_Or(TUniquePtr<FComposableStrandFilterOp> InLeftOp, TUniquePtr<FComposableStrandFilterOp> InRightOp);
+    virtual bool FilterStrand(const FHairStrandAttribsRefProxy& StrandAttributes) override;
+
+private:
+    TUniquePtr<FComposableStrandFilterOp> LeftOp;
+    TUniquePtr<FComposableStrandFilterOp> RightOp;
+};
+
+class FStrandFilterOp_InCardGroup : public FComposableStrandFilterOp
+{
+public:
+	FStrandFilterOp_InCardGroup(const TSet<FName>& InCardGroupIds);
+    virtual bool FilterStrand(const FHairStrandAttribsRefProxy& StrandAttributes) override;
+
+private:
+	TSet<FName> CardGroupNames;
+};
diff --git a/Engine/Plugins/Experimental/HairCardGenerator/Source/HairCardGeneratorEditor/Private/HairCardGenStrandNCSInterpolator.cpp b/Engine/Plugins/Experimental/HairCardGenerator/Source/HairCardGeneratorEditor/Private/HairCardGenStrandNCSInterpolator.cpp
new file mode 100644
--- /dev/null
+++ b/Engine/Plugins/Experimental/HairCardGenerator/Source/HairCardGeneratorEditor/Private/HairCardGenStrandNCSInterpolator.cpp
@@ -0,0 +1,62 @@
+// Copyright Epic Games, Inc. All Rights Reserved.
+
+#include "HairCardGenStrandNCSInterpolator.h"
+#include <iostream>
+
+FHairCardGenStrandNSCInterpolator::FHairCardGenStrandNSCInterpolator(const TArray<FVector>& Positions, const TArray<float>& Widths)
+{
+    int NumPoints = Positions.Num();
+
+	TArray<float> X, Y, Z;
+	X.SetNum(NumPoints);
+	Y.SetNum(NumPoints);
+	Z.SetNum(NumPoints);
+
+	for (int i = 0; i < NumPoints; i++) {
+		X[i] = Positions[i].X;
+        Y[i] = Positions[i].Y;
+        Z[i] = Positions[i].Z;
+	}
+
+    TArray<float> DistancesFromStart;
+	DistancesFromStart.SetNum(NumPoints);
+    DistancesFromStart[0] = 0;
+	for (int i = 1; i < NumPoints; ++i)
+	{
+		DistancesFromStart[i] = DistancesFromStart[i - 1] + FVector::Dist(Positions[i], Positions[i - 1]);
+	}
+	for (int i = 1; i < NumPoints; ++i)
+	{
+		DistancesFromStart[i] /= DistancesFromStart[NumPoints - 1];
+	}
+
+	Splines.Reserve(4);
+
+	Splines.Add(FHairCardGenNaturalCubicSplines(DistancesFromStart, X));
+	Splines.Add(FHairCardGenNaturalCubicSplines(DistancesFromStart, Y));
+	Splines.Add(FHairCardGenNaturalCubicSplines(DistancesFromStart, Z));
+	Splines.Add(FHairCardGenNaturalCubicSplines(DistancesFromStart, Widths));
+
+	Splines.Shrink();
+}
+
+FHairCardGenStrandNSCInterpolator::StrandInterpolationResult FHairCardGenStrandNSCInterpolator::GetInterpolatedStrand(const int NumInterpolatedPoints)
+{
+	StrandInterpolationResult Result;
+
+	Result.Positions.SetNum(NumInterpolatedPoints * 3);
+	Result.Widths.SetNum(NumInterpolatedPoints);
+
+	for (int i = 0; i < NumInterpolatedPoints; i++)
+	{
+		float T = float(i) / float(NumInterpolatedPoints - 1);
+
+		for (int j = 0; j < 3; j++)
+		{
+			Result.Positions[i * 3 + j] = Splines[j](T);
+		}
+		Result.Widths[i] = Splines[3](T);
+	}
+
+	return Result;
+}
\ No newline at end of file
diff --git a/Engine/Plugins/Experimental/HairCardGenerator/Source/HairCardGeneratorEditor/Private/HairCardGenStrandNCSInterpolator.h b/Engine/Plugins/Experimental/HairCardGenerator/Source/HairCardGeneratorEditor/Private/HairCardGenStrandNCSInterpolator.h
new file mode 100644
--- /dev/null
+++ b/Engine/Plugins/Experimental/HairCardGenerator/Source/HairCardGeneratorEditor/Private/HairCardGenStrandNCSInterpolator.h
@@ -0,0 +1,24 @@
+// Copyright Epic Games, Inc. All Rights Reserved.
+
+#pragma once
+
+#include "Containers/Array.h"
+#include "Math/Vector.h"
+#include "HairCardGenNaturalCubicSplines.h"
+
+class FHairCardGenStrandNSCInterpolator
+{
+private:
+    TArray<FHairCardGenNaturalCubicSplines> Splines;
+
+public:
+    struct StrandInterpolationResult
+    {
+        TArray<float> Positions;
+        TArray<float> Widths;
+    };
+
+public:
+	FHairCardGenStrandNSCInterpolator(const TArray<FVector>& Positions, const TArray<float>& Widths);
+    StrandInterpolationResult GetInterpolatedStrand(const int NumInterpolatedPoints);
+}; 
diff --git a/Engine/Plugins/Experimental/HairCardGenerator/Source/HairCardGeneratorEditor/Private/HairCardGeneratorEditorModule.cpp b/Engine/Plugins/Experimental/HairCardGenerator/Source/HairCardGeneratorEditor/Private/HairCardGeneratorEditorModule.cpp
new file mode 100644
--- /dev/null
+++ b/Engine/Plugins/Experimental/HairCardGenerator/Source/HairCardGeneratorEditor/Private/HairCardGeneratorEditorModule.cpp
@@ -0,0 +1,860 @@
+// Copyright Epic Games, Inc. All Rights Reserved.
+
+#include "IHairCardGeneratorEditor.h"
+#include "Engine/Texture2D.h"
+#include "Factories/TextureFactory.h"
+#include "HairCardGeneratorLog.h"
+#include "HairCardGenOptionsWindow.h"
+#include "HairCardGeneratorPluginSettings.h"
+#include "HairCardGenControllerBase.h"
+#include "HairCardGeneratorEditorSettings.h"
+#include "HairCardGenCardSubdivider.h"
+#include "HairCardGenStrandNCSInterpolator.h"
+#include "HairCardGenStrandFilter.h"
+#include "HairCardGenSettingsDetail.h"
+
+#include "AssetToolsModule.h"
+#include "PropertyEditorModule.h"
+#include "AutomatedAssetImportData.h"
+#include "GroomAsset.h"
+#include "GroomCacheData.h" // for EGroomBasisType, EGroomCurveType
+#include "HairDescription.h"
+#include "ObjectTools.h"
+#include "AssetRegistry/AssetRegistryModule.h" // for GetAssetByObjectPath()
+#include "Engine/StaticMesh.h"
+#include "HAL/IConsoleManager.h"
+#include "HAL/LowLevelMemTracker.h"
+#include "HAL/PlatformFileManager.h"
+#include "Misc/Paths.h"
+#include "Misc/PackageName.h"
+#include "Misc/ScopedSlowTask.h"
+#include "UObject/UObjectHash.h" // for GetObjectsOfClass()
+#include "UObject/WeakObjectPtrTemplates.h"
+#include "Framework/Notifications/NotificationManager.h"
+#include "Widgets/Notifications/SNotificationList.h"
+
+DEFINE_LOG_CATEGORY(LogHairCardGenerator)
+#define LOCTEXT_NAMESPACE "HairCardGeneratorEd"
+
+LLM_DEFINE_TAG(HairCardGenerator);
+
+/* Used for keeping track of the editor settings for importing generated data
+*/
+struct FHairCardImportSettings
+{
+	FString MeshName;
+	TMap<EHairAtlasTextureType, FString> TextureNames;
+};
+
+
+/* FHairCardGeneratorHelper Declaration
+ *****************************************************************************/
+
+// Singleton class for caching a pointer to the hair card generator
+class FHairCardGenController
+{
+public:
+	static TObjectPtr<UHairCardGenControllerBase> GetHairCardGenerator();
+
+private:
+	static UGroomAsset* GroomAsset;
+	static TWeakObjectPtr<UHairCardGenControllerBase> HairCardGenController;
+};
+
+TWeakObjectPtr<UHairCardGenControllerBase> FHairCardGenController::HairCardGenController = nullptr;
+
+TObjectPtr<UHairCardGenControllerBase> FHairCardGenController::GetHairCardGenerator()
+{
+	if ( HairCardGenController.IsValid() )
+		return HairCardGenController.Get();
+
+	// Find the HairCardGenController in Python (or wherever it may be implemented -- could be moved to C++/Blueprints/etc.)
+	TArray<UObject*> HairCardGenControllers;
+	GetObjectsOfClass(UHairCardGenControllerBase::StaticClass(), HairCardGenControllers);
+
+	if ( HairCardGenControllers.Num() == 0 )
+	{
+		UE_LOG(LogHairCardGenerator, Error, TEXT("Failed to find an active hair card controller."));
+		return nullptr;
+	}
+
+	if ( HairCardGenControllers.Num() != 1 )
+		UE_LOG(LogHairCardGenerator, Warning, TEXT("More than one active hair-card controller (using first controller found)."))
+
+	HairCardGenController = TWeakObjectPtr<UHairCardGenControllerBase>(static_cast<UHairCardGenControllerBase*>(HairCardGenControllers[0]));
+	return HairCardGenController.Get();
+}
+
+/* HairCardGeneratorEditor_Impl Declaration
+ *****************************************************************************/
+
+// Static namespace with helper functions for commands/generation tasks
+namespace HairCardGeneratorEditor_Impl
+{
+	static TObjectPtr<UGroomAsset> ParseGroomAssetArg(const TArray<FString>& Args, bool bOptional);
+	static FHairGroupsCardsSourceDescription* FindCardSourceDescription(TObjectPtr<UGroomAsset> Groom, int32 LODIndex);
+	static TObjectPtr<UHairCardGeneratorPluginSettings> CreateGroomSettingsLast(TObjectPtr<UGroomAsset> Groom, const FHairGroupsCardsSourceDescription& SourceDesc);
+	static TObjectPtr<UHairCardGeneratorPluginSettings> CreateGroomSettingsDefault(TObjectPtr<UGroomAsset> Groom, int32 LODIndex, int32 PhysGroupIndex);
+
+	static void DbgCmdCacheGroom(const TArray<FString>& Args);
+
+	static bool LoadGroomData(UGroomAsset* GroomAsset, const bool SaveCached);
+	static bool BuildHairCardGenGroomData(const int PointsPerStrand, const FHairDescription& InGroomData, FHairCardGen_GroomData& Out);
+	static bool GenerateCardsForCardGroup(TObjectPtr<const UHairCardGeneratorPluginSettings> Settings, int index, uint8 GenFlags);
+	static bool GenerateCardsGeometry(TObjectPtr<const UHairCardGeneratorPluginSettings> Settings, int index);
+	static void GetTextureImportNames(UHairCardGeneratorPluginSettings* Settings, FHairCardImportSettings& ImportSettings);
+	static TArray<UObject*> ImportHairCardTextures(const FString& ImportSrc, const FString& ImportDst, const FHairCardImportSettings& ImportSettings);
+	static UStaticMesh* CreateMeshInPackage(const FString& DestinationPackage, const FString& BaseFilename, int GroupID);
+	static int32 AssignImportedTextures(const TArray<UObject*>& ImportedAssets, const FHairCardImportSettings& ImportSettings, FHairGroupsCardsSourceDescription& Assignee);
+	static void CreateFinishNotification(bool bSuccess, const FText& NotificationText);
+
+	static TWeakObjectPtr<UGroomAsset> LastGroom = nullptr;
+}
+
+
+static FAutoConsoleCommand CCmdLoadGroom(
+	TEXT("HairCardSolver.CacheGroom"),
+	TEXT(""),
+	FConsoleCommandWithArgsDelegate::CreateStatic(HairCardGeneratorEditor_Impl::DbgCmdCacheGroom));
+
+/* HairCardGeneratorEditor_Impl Definition
+ *****************************************************************************/
+
+// Always assume the first argument is a (sometimes optional) groom argument and get UGroomAsset, if possible
+static TObjectPtr<UGroomAsset> HairCardGeneratorEditor_Impl::ParseGroomAssetArg(const TArray<FString>& Args, bool bOptional)
+{
+	if ( Args.Num() < 1 )
+		return nullptr;
+
+	FString AssetPath = Args[0];
+	FSoftObjectPath CheckPath(AssetPath);
+
+	if ( !CheckPath.IsValid() )
+	{
+		UE_CLOG(!bOptional, LogHairCardGenerator, Error, TEXT("The path is not a valid asset path."));
+		return nullptr;
+	}
+
+	FAssetRegistryModule& AssetRegistryModule = FModuleManager::LoadModuleChecked<FAssetRegistryModule>("AssetRegistry");
+	FAssetData ExistingAsset = AssetRegistryModule.Get().GetAssetByObjectPath(CheckPath);
+	if (!ExistingAsset.IsValid())
+	{
+		UE_CLOG(!bOptional, LogHairCardGenerator, Error, TEXT("The asset path cannot be found in the registry."));
+		return nullptr;
+	}
+
+
+	if(ExistingAsset.GetClass() != UGroomAsset::StaticClass())
+	{
+		UE_CLOG(!bOptional, LogHairCardGenerator, Error, TEXT("The specified asset is not a groom asset."));
+		return nullptr;
+	}
+
+	return LoadObject<UGroomAsset>(nullptr, *AssetPath);
+}
+
+static FHairGroupsCardsSourceDescription* HairCardGeneratorEditor_Impl::FindCardSourceDescription(TObjectPtr<UGroomAsset> Groom, int32 LODIndex)
+{
+	// Try to load previous settings from a matching (LOD) card source description slot, if compatible
+	for ( FHairGroupsCardsSourceDescription& Desc : Groom->GetHairGroupsCards() )
+	{
+		if ( Desc.LODIndex == LODIndex)
+		{
+			return &Desc;
+		}
+	}
+
+	return nullptr;
+}
+
+static TObjectPtr<UHairCardGeneratorPluginSettings> HairCardGeneratorEditor_Impl::CreateGroomSettingsLast(TObjectPtr<UGroomAsset> Groom, const FHairGroupsCardsSourceDescription& SourceDesc)
+{
+	TObjectPtr<UHairCardGeneratorPluginSettings> NewSettings = NewObject<UHairCardGeneratorPluginSettings>(/*Outer =*/Groom);
+	NewSettings->SetSource(Groom, SourceDesc.LODIndex, SourceDesc.GroupIndex);
+	NewSettings->ResetToDefault();
+
+	// First try to reset from Json file (then fallback to previous settings if they exist)
+	if ( NewSettings->ResetFromSettingsJson() )
+		return NewSettings;
+
+	NewSettings->ResetFromSourceDescription(SourceDesc);
+	return NewSettings;
+}
+
+static TObjectPtr<UHairCardGeneratorPluginSettings> HairCardGeneratorEditor_Impl::CreateGroomSettingsDefault(TObjectPtr<UGroomAsset> Groom, int32 LODIndex, int32 PhysGroupIndex)
+{
+	TObjectPtr<UHairCardGeneratorPluginSettings> NewSettings = NewObject<UHairCardGeneratorPluginSettings>(/*Outer =*/Groom);
+	NewSettings->SetSource(Groom, LODIndex, PhysGroupIndex);
+	NewSettings->ResetToDefault();
+
+	return NewSettings;
+}
+
+static void HairCardGeneratorEditor_Impl::DbgCmdCacheGroom(const TArray<FString>& Args)
+{
+	TObjectPtr<UGroomAsset> GroomAsset = HairCardGeneratorEditor_Impl::ParseGroomAssetArg(Args, false);
+	if ( !GroomAsset )
+		return;
+
+	if ( !HairCardGeneratorEditor_Impl::LoadGroomData(GroomAsset, true) )
+	{
+		UE_LOG(LogHairCardGenerator, Error, TEXT("Failed to cache groom asset."));
+	}
+}
+
+static bool HairCardGeneratorEditor_Impl::LoadGroomData(UGroomAsset* NewGroomAsset, bool SaveCached = false)
+{
+	if (HairCardGeneratorEditor_Impl::LastGroom == NewGroomAsset && !SaveCached)
+	{
+		UE_LOG(LogHairCardGenerator, Warning, TEXT("Groom already loaded."));
+		return true;
+	}
+	else
+	{
+		TObjectPtr<UHairCardGenControllerBase> CardGenController = FHairCardGenController::GetHairCardGenerator();
+		if ( !CardGenController )
+			return false;
+
+		FHairCardGen_GroomData RestructuredData;
+		if ( !HairCardGeneratorEditor_Impl::BuildHairCardGenGroomData(CardGenController->GetPointsPerCurve(), NewGroomAsset->GetHairDescription(), RestructuredData))
+			return false;
+
+		FScopedSlowTask SlowTask(0, LOCTEXT("GeneratingHairCards", "Loading Groom..."));
+		SlowTask.MakeDialog(/*bShowCancelButton =*/true);
+
+		FString CachedGroomsPath = FPaths::ProjectIntermediateDir() / TEXT("GroomHairCardGen") / TEXT("CachedGrooms");
+		if (!CardGenController->LoadGroomData(RestructuredData, NewGroomAsset->GetName(), CachedGroomsPath, SaveCached))
+		{
+			return false;
+		}
+		else
+		{
+			HairCardGeneratorEditor_Impl::LastGroom = NewGroomAsset;
+			return true;
+		}
+	}
+}
+
+static bool HairCardGeneratorEditor_Impl::BuildHairCardGenGroomData(const int PointsPerStrand, const FHairDescription& InGroomData, FHairCardGen_GroomData& Out)
+{
+	const int NumStrands = InGroomData.GetNumStrands();
+	Out.Strands.Reserve(NumStrands);
+
+	TGroomAttributesConstRef<float> GroomHairWidth = InGroomData.GroomAttributes().GetAttributesRef<float>(HairAttribute::Groom::Width);
+	const FGroomID GroomID = FGroomID(0);
+	const float Groom_DefaultWidth = GroomHairWidth.IsValid() ? GroomHairWidth[GroomID] : 0.f;
+
+	FStrandAttributesRange TraverseStrands(InGroomData, FHairStrandAttribsConstIterator::DefaultIgnore);
+	if ( !TraverseStrands.IsValid() )
+		return false;
+
+	UE_CLOG(!TraverseStrands.First().HasGroupID(), LogHairCardGenerator, Warning, TEXT("Groom description has no physics group info, may not generate correctly"));
+
+	const FHairStrandAttribsRefProxy& FirstStrand = TraverseStrands.First();
+	UE_CLOG((!FirstStrand.HasCurveType() || !FirstStrand.HasBasisType() || FirstStrand.HasKnots()), LogHairCardGenerator, Warning, TEXT("Invalid strand curve specification, treating as linear"));
+
+	EGroomCurveType CurveType = FirstStrand.GetCurveType();
+	EGroomBasisType BasisType = FirstStrand.GetBasisType();
+
+	if (CurveType == EGroomCurveType::VariableOrder)
+	{
+		UE_LOG(LogHairCardGenerator, Warning, TEXT("VariableOrder curves not supported, treating as linear"));
+
+		CurveType = EGroomCurveType::Linear;
+		BasisType = EGroomBasisType::NoBasis;
+	}
+
+	Out.CurveType = StaticEnum<EGroomCurveType>()->GetNameStringByValue((int64)CurveType);
+	Out.BasisType = StaticEnum<EGroomBasisType>()->GetNameStringByValue((int64)BasisType);
+
+	Out.VertexPositions.Reserve(NumStrands * PointsPerStrand * 3);
+	Out.VertexWidths.Reserve(NumStrands * PointsPerStrand);
+
+	for ( const FHairStrandAttribsRefProxy& StrandAttribs : TraverseStrands )
+	{
+		FHairCardGen_StrandData& OutStrandData = Out.Strands.AddDefaulted_GetRef();
+		OutStrandData.GroupID = StrandAttribs.HasGroupID() ? StrandAttribs.GetGroupID() : -1;
+
+		const float Strand_DefaultWidth = StrandAttribs.HasStrandWidth() ? StrandAttribs.GetStrandWidth() : Groom_DefaultWidth;
+
+		TArray<FVector> VertexPositions;
+		VertexPositions.SetNum(StrandAttribs.NumVerts());
+		TArray<float> VertexWidths;
+		VertexWidths.SetNum(StrandAttribs.NumVerts());
+
+		for (int StrandVertIndex = 0; StrandVertIndex < StrandAttribs.NumVerts(); ++StrandVertIndex)
+		{
+			VertexPositions[StrandVertIndex] = FVector(StrandAttribs.GetVertexPosition(StrandVertIndex));
+			VertexWidths[StrandVertIndex] = StrandAttribs.HasVertexWidths() ? StrandAttribs.GetVertexWidth(StrandVertIndex) : Strand_DefaultWidth;
+		}
+
+		FHairCardGenStrandNSCInterpolator::StrandInterpolationResult InterpolatedStrand = FHairCardGenStrandNSCInterpolator(VertexPositions, VertexWidths).GetInterpolatedStrand(PointsPerStrand);
+
+		Out.VertexPositions.Append(InterpolatedStrand.Positions);
+		Out.VertexWidths.Append(InterpolatedStrand.Widths);
+	}
+
+	Out.Strands.Shrink();
+	Out.VertexPositions.Shrink();
+	Out.VertexWidths.Shrink();
+
+	Out.HairlineGroupID = INDEX_NONE;
+
+	return true;
+}
+
+static bool HairCardGeneratorEditor_Impl::GenerateCardsGeometry(TObjectPtr<const UHairCardGeneratorPluginSettings> Settings, int index)
+{
+	TObjectPtr<UHairCardGenControllerBase> CardGenController = FHairCardGenController::GetHairCardGenerator();
+	if ( !CardGenController )
+	{
+		return false;
+	}
+
+	FScopedSlowTask SlowTask(10, FText::Format(LOCTEXT("GeneratingHairCards.GeometryOptimization", "Generating cards geometry"), index));
+
+	SlowTask.EnterProgressFrame(1);
+	TArray<int> NumCards = CardGenController->SetOptimizations(index);
+
+	if (NumCards.Num() == 0)
+	{
+		UE_LOG(LogHairCardGenerator, Error, TEXT("Card geometry generation failed for settings group %d. Check the log for details."), index);
+		SlowTask.EnterProgressFrame(9);
+		return false;
+	}
+
+	int TotalNumCards = 0;
+	for (int CurrentNumCards : NumCards) TotalNumCards += CurrentNumCards;
+
+	float AvgLength = 0;
+	float AvgCurvRadius = 0;
+
+	TArray<MatrixXf> AvgCurves;
+	AvgCurves.SetNum(TotalNumCards);
+
+	int CardId = 0;
+	for (int Id = 0; Id < NumCards.Num(); Id++)
+	{
+		for (int Cid = 0; Cid < NumCards[Id]; Cid++)
+		{
+			const TArray<float> AvgCurvePoints = CardGenController->GetAverageCurve(Id, Cid);
+
+			if (AvgCurvePoints.Num() == 0)
+			{
+				UE_LOG(LogHairCardGenerator, Error, TEXT("Card geometry generation failed for settings group %d. Check the log for details."), index);
+				SlowTask.EnterProgressFrame(9);
+				return false;
+			}
+
+			AvgCurves[CardId] = FHairCardGenCardSubdivider::TArrayToMatrixXf(AvgCurvePoints);
+			AvgLength += FHairCardGenCardSubdivider::GetCurveLength(AvgCurves[CardId]);
+			AvgCurvRadius += FHairCardGenCardSubdivider::GetAverageCurvatureRadius(AvgCurves[CardId]);
+
+			CardId++;
+		}
+	}
+
+	AvgLength /= (float)TotalNumCards;
+	AvgCurvRadius /= (float)TotalNumCards;
+
+	auto GetSubdToleranceFromTriCount = [AvgLength, AvgCurvRadius, TotalNumCards](int TriCount)
+	{
+		const float NumSubdivisions = 0.5 * (float)TriCount / (float)TotalNumCards;
+		return FHairCardGenCardSubdivider::GetToleranceFromLengthAndCurvRadius(AvgLength, AvgCurvRadius, NumSubdivisions);
+	};
+
+	FHairCardGenCardSubdivider Subdivider(
+		GetSubdToleranceFromTriCount(Settings->GetFilterGroupSettings(index)->TargetTriangleCount),
+		Settings->GetFilterGroupSettings(index)->UseAdaptiveSubdivision,
+		Settings->GetFilterGroupSettings(index)->MaxVerticalSegmentsPerCard);
+
+	auto CountTriangles = [&Subdivider, &AvgCurves]() 
+	{
+		int TriCount = 0;
+		for (int CardId = 0; CardId < AvgCurves.Num(); CardId++) TriCount += 2 * (Subdivider.GetSubdivisionPoints(AvgCurves[CardId]).Num() / 3 - 1);
+		return TriCount;
+	};
+
+	const int InitialTargetTriCount = Settings->GetFilterGroupSettings(index)->TargetTriangleCount;
+	const int InitialTriCount = CountTriangles();
+
+	const int ShiftedTargetTriCount = 2 * InitialTargetTriCount - InitialTriCount;
+	Subdivider.SetSubdTolerance(GetSubdToleranceFromTriCount(ShiftedTargetTriCount));
+	const int ShiftedTriCount = CountTriangles();
+
+	if (InitialTriCount != ShiftedTriCount)
+	{
+		const int CorrectedTargetTriCount = InitialTargetTriCount + int(float(ShiftedTargetTriCount - InitialTargetTriCount) / float(ShiftedTriCount - InitialTriCount) * float(InitialTargetTriCount - InitialTriCount));
+		Subdivider.SetSubdTolerance(GetSubdToleranceFromTriCount(CorrectedTargetTriCount));
+	}
+
+	CardId = 0;
+	for (int Id = 0; Id < NumCards.Num(); Id++)
+	{
+		for (int Cid = 0; Cid < NumCards[Id]; Cid++)
+		{
+			TArray<float> SubdividedPoints = Subdivider.GetSubdivisionPoints(AvgCurves[CardId]);
+			CardGenController->SetInterpolatedAvgCurve(Id, Cid, SubdividedPoints);
+
+			CardId++;
+		}
+	}
+
+	SlowTask.EnterProgressFrame(9);
+	if (!CardGenController->GenerateCardsGeometry())
+	{
+		UE_LOG(LogHairCardGenerator, Error, TEXT("Card geometry generation failed for settings group %d. Check the log for details."), index);
+		return false;
+	}
+
+	return true;
+}
+
+static bool HairCardGeneratorEditor_Impl::GenerateCardsForCardGroup(TObjectPtr<const UHairCardGeneratorPluginSettings> Settings, int index, uint8 GenFlags)
+{
+	TObjectPtr<UHairCardGenControllerBase> CardGenController = FHairCardGenController::GetHairCardGenerator();
+	if ( !CardGenController )
+	{
+		return false;
+	}
+
+	FScopedSlowTask SlowTask(10, FText::Format(LOCTEXT("GeneratingHairCards", "Generating Hair Cards for settings group {0}..."), index));
+	SlowTask.MakeDialog(/*bShowCancelButton =*/true);
+
+	SlowTask.EnterProgressFrame(1);
+	if ( UHairCardGeneratorPluginSettings::CheckGenerationFlags(GenFlags, EHairCardGenerationPipeline::StrandClustering) && !Settings->bReduceCardsFromPreviousLOD )
+	{
+		if ( !CardGenController->GenerateClumps(index) )
+		{
+			UE_LOG(LogHairCardGenerator, Error, TEXT("Strand clustering failed for settings group %d. Check the log for details."), index);
+			SlowTask.EnterProgressFrame(9);
+			return false;
+		}
+	}
+
+	SlowTask.EnterProgressFrame(5);
+	if ( UHairCardGeneratorPluginSettings::CheckGenerationFlags(GenFlags, EHairCardGenerationPipeline::GeometryGeneration) )
+	{
+		if ( !HairCardGeneratorEditor_Impl::GenerateCardsGeometry(Settings, index) )
+		{
+			SlowTask.EnterProgressFrame(4);
+			return false;
+		}
+	}
+
+	SlowTask.EnterProgressFrame(4);
+	if ( UHairCardGeneratorPluginSettings::CheckGenerationFlags(GenFlags, EHairCardGenerationPipeline::TextureClustering) && !Settings->bReduceCardsFromPreviousLOD )
+	{
+		if (!CardGenController->ClusterTextures(index))
+		{
+			UE_LOG(LogHairCardGenerator, Error, TEXT("Texture clustering failed for settings group %d. Check the log for details."), index)
+			return false;
+		}
+	}
+
+	return true;
+}
+
+static void HairCardGeneratorEditor_Impl::GetTextureImportNames(UHairCardGeneratorPluginSettings* Settings, FHairCardImportSettings& ImportSettings)
+{
+	FString TextureBasename = Settings->GetTextureImportBaseName();
+
+	ImportSettings.TextureNames.Empty();
+	ImportSettings.MeshName = Settings->BaseFilename;
+
+	FFormatNamedArguments FilenameFormatArgs;
+	FilenameFormatArgs.Add(TEXT("groom_name"), FText::FromString(Settings->GetGroomName()));
+	FilenameFormatArgs.Add(TEXT("lod"), Settings->GetLODIndex());
+	FilenameFormatArgs.Add(TEXT("mesh_name"), FText::FromString(TextureBasename));
+
+	const UHairCardGeneratorEditorSettings* ConfigSettings = GetDefault<UHairCardGeneratorEditorSettings>();
+	ImportSettings.TextureNames.Reserve(ConfigSettings->HairCardTextureNameFormats.Num());
+	for (auto TextureNameIt : ConfigSettings->HairCardTextureNameFormats)
+	{
+		FString FormattedTextureName = FText::Format(FText::FromString(TextureNameIt.Value), FilenameFormatArgs).ToString();
+		if (FPaths::ValidatePath(FormattedTextureName))
+		{
+			ImportSettings.TextureNames.Add(TextureNameIt.Key, FormattedTextureName);
+		}
+		else
+		{
+			UE_LOG(LogHairCardGenerator, Warning, TEXT("Invalid texture name config: '%s'"), *TextureNameIt.Value);
+		}
+	}
+}
+
+static TArray<UObject*> HairCardGeneratorEditor_Impl::ImportHairCardTextures(const FString& ImportSrc, const FString& ImportDst, const FHairCardImportSettings& ImportSettings)
+{
+	// Create a texture factory for importing tangents/attributes with linear colorspace
+	UTextureFactory* LinearTextureFactory = NewObject<UTextureFactory>();
+	LinearTextureFactory->ColorSpaceMode = ETextureSourceColorSpace::Linear;
+
+	UAutomatedAssetImportData* AutoImportLinearTextures = NewObject<UAutomatedAssetImportData>();
+	AutoImportLinearTextures->DestinationPath = ImportDst;
+	AutoImportLinearTextures->bReplaceExisting = true;
+	AutoImportLinearTextures->bSkipReadOnly = false;
+	AutoImportLinearTextures->Factory = LinearTextureFactory;
+
+	UAutomatedAssetImportData* AutomatedImportData = NewObject<UAutomatedAssetImportData>();
+	AutomatedImportData->DestinationPath = ImportDst;
+	AutomatedImportData->bReplaceExisting = true;
+	AutomatedImportData->bSkipReadOnly = false;
+
+	IPlatformFile& FileManager = FPlatformFileManager::Get().GetPlatformFile();
+	FileManager.IterateDirectory(*ImportSrc, [AutomatedImportData, AutoImportLinearTextures, ImportSettings](const TCHAR* FilenameOrDirectory, const bool bIsDirectory)->bool
+		{
+			if (!bIsDirectory)
+			{
+				FString basename = FPaths::GetBaseFilename(FilenameOrDirectory, true);
+				const EHairAtlasTextureType* pType = ImportSettings.TextureNames.FindKey(basename);
+				if ( pType && (*pType == EHairAtlasTextureType::Tangent || *pType == EHairAtlasTextureType::Attribute ) )
+					AutoImportLinearTextures->Filenames.Add(FString(FilenameOrDirectory));
+				else
+					AutomatedImportData->Filenames.Add(FilenameOrDirectory);
+			}
+			return true;
+		}
+	);
+
+	TArray<UObject*> ImportedAssets = FAssetToolsModule::GetModule().Get().ImportAssetsAutomated(AutomatedImportData);
+	ImportedAssets.Append(FAssetToolsModule::GetModule().Get().ImportAssetsAutomated(AutoImportLinearTextures));
+	return ImportedAssets;
+}
+
+static UStaticMesh* HairCardGeneratorEditor_Impl::CreateMeshInPackage(const FString& DestinationPackage, const FString& BaseFilename, int GroupID)
+{
+	FString MeshName = BaseFilename;
+	if ( GroupID >= 0 )
+	{
+		MeshName += FString::Printf(TEXT("_PG%d"), GroupID);
+	}
+
+	FAssetRegistryModule& AssetRegistryModule = FModuleManager::LoadModuleChecked<FAssetRegistryModule>("AssetRegistry");
+
+	FString PackageDest = ObjectTools::SanitizeInvalidChars(FPaths::Combine(*DestinationPackage, *MeshName), INVALID_LONGPACKAGE_CHARACTERS);
+	FString ObjectName = ObjectTools::SanitizeObjectName(MeshName);
+
+	UPackage* PackageOuter = CreatePackage(*PackageDest);
+	TObjectPtr<UObject> ExistingMesh = FindObject<UStaticMesh>(PackageOuter, *ObjectName);
+	if ( !ExistingMesh )
+	{
+		UStaticMesh* NewMesh = NewObject<UStaticMesh>(PackageOuter, UStaticMesh::StaticClass(), *ObjectName, RF_Public | RF_Standalone);
+		// Add asset path to registry to get folder reference in content browser
+		AssetRegistryModule.Get().AssetCreated(NewMesh);
+
+		return NewMesh;
+	}
+
+	if ( ExistingMesh->GetClass()->IsChildOf(UStaticMesh::StaticClass()) )
+	{
+		// Free any RHI resources for existing mesh before we re-create in place.
+		ExistingMesh->PreEditChange(nullptr);
+		UStaticMesh* NewMesh = NewObject<UStaticMesh>(PackageOuter, UStaticMesh::StaticClass(), *ObjectName, RF_Public | RF_Standalone);
+		AssetRegistryModule.Get().AssetCreated(NewMesh);
+
+		return NewMesh;
+	}
+
+	UE_LOG(LogHairCardGenerator, Error, TEXT("Non-mesh asset ('%s') already exists in package: '%s'"), *ObjectName, *PackageDest);
+
+	return nullptr;
+}
+
+static int32 HairCardGeneratorEditor_Impl::AssignImportedTextures(const TArray<UObject*>& ImportedAssets, const FHairCardImportSettings& ImportSettings, FHairGroupsCardsSourceDescription& Assignee)
+{
+	int32 AssignedAssets = 0;
+	for (UObject* ImportedObject : ImportedAssets)
+	{
+		TObjectPtr<UTexture2D> TextureImport = Cast<UTexture2D>(ImportedObject);
+		if ( TextureImport )
+		{
+			bool bFound = false;
+			for (auto NameTemplateIt : ImportSettings.TextureNames)
+			{
+				if (TextureImport->GetName().Equals(NameTemplateIt.Value))
+				{
+					Assignee.Textures.SetTexture(NameTemplateIt.Key, TextureImport);
+					bFound = true;
+					++AssignedAssets;
+					break;
+				}
+			}
+
+			UE_CLOG(!bFound, LogHairCardGenerator, Warning, TEXT("Unable to determine which slot imported texture ('%s') is for. Check your config."), *TextureImport->GetName());
+		}
+		else
+		{
+			UE_LOG(LogHairCardGenerator, Warning, TEXT("Unused asset imported along with hair card mesh: '%s'"), *ImportedObject->GetName());
+		}
+	}
+
+	return AssignedAssets;
+}
+
+static void HairCardGeneratorEditor_Impl::CreateFinishNotification(bool bSuccess, const FText& NotificationText)
+{
+	FNotificationInfo Info(NotificationText);
+
+	Info.bFireAndForget = true;
+	Info.bUseLargeFont = true;
+	Info.bUseSuccessFailIcons = true;
+
+	Info.FadeOutDuration = 4.0f;
+	Info.ExpireDuration = Info.FadeOutDuration;
+
+	SNotificationItem::ECompletionState State = (bSuccess) ? SNotificationItem::CS_Success : SNotificationItem::CS_Fail;
+
+	TSharedPtr<SNotificationItem> NotificationPtr = FSlateNotificationManager::Get().AddNotification(Info);
+	NotificationPtr->SetCompletionState(State);
+}
+
+/* FHairCardGeneratorEditorModule
+ *****************************************************************************/
+
+class FHairCardGeneratorEditorModule : public IHairCardGeneratorEditor
+{
+public:
+	//~ Begin IModuleInterface API
+	virtual void StartupModule() override;
+	virtual void ShutdownModule() override;
+	//~ End IModuleInterface API
+
+	//~ Begin IHairCardGenerator API
+	virtual bool GenerateHairCardsForLOD(UGroomAsset* Groom, FHairGroupsCardsSourceDescription& CardsDesc) override;
+	virtual bool IsCompatibleSettings(UHairCardGenerationSettings* OldSettings) override;
+	//~ End IHairCardGenerator API
+};
+
+void FHairCardGeneratorEditorModule::StartupModule()
+{
+	LLM_SCOPE_BYTAG(HairCardGenerator)
+
+	HairCardGenerator_Utils::RegisterModularHairCardGenerator(this);
+
+	// Register custom detail panel handler for hair card settings
+	FPropertyEditorModule& PropertyModule = FModuleManager::LoadModuleChecked<FPropertyEditorModule>("PropertyEditor");
+	PropertyModule.RegisterCustomClassLayout(UHairCardGeneratorPluginSettings::StaticClass()->GetFName(), FOnGetDetailCustomizationInstance::CreateStatic(&FHairCardSettingsDetailCustomization::MakeInstance));
+}
+
+void FHairCardGeneratorEditorModule::ShutdownModule()
+{
+	HairCardGenerator_Utils::UnregisterModularHairCardGenerator(this);
+
+	// Register custom detail panel handler for hair card settings
+	FPropertyEditorModule& PropertyModule = FModuleManager::LoadModuleChecked<FPropertyEditorModule>("PropertyEditor");
+	PropertyModule.UnregisterCustomClassLayout(UHairCardGeneratorPluginSettings::StaticClass()->GetFName());
+}
+
+bool FHairCardGeneratorEditorModule::GenerateHairCardsForLOD(UGroomAsset* NewGroomAsset, FHairGroupsCardsSourceDescription& CardsDesc)
+{
+	// Find the HairCardGenController in Python (or wherever it may be implemented -- could be moved to C++/Blueprints/etc.)
+	TObjectPtr<UHairCardGenControllerBase> CardGenController = FHairCardGenController::GetHairCardGenerator();
+	if ( !CardGenController )
+	{
+		UE_LOG(LogHairCardGenerator, Error, TEXT("No hair card generator controller found, cannot generate hair cards!"));
+		return false;
+	}
+
+	UE_CLOG(CardsDesc.GenerationSettings && !IsCompatibleSettings(CardsDesc.GenerationSettings), LogHairCardGenerator, Warning, TEXT("Old hair-card generation settings (from a different generator) will be discarded when you run this generator."));
+
+	TObjectPtr<UHairCardGeneratorPluginSettings>& GenerationSettings = CardGenController->GetGroomSettings(NewGroomAsset, CardsDesc.LODIndex, CardsDesc.GroupIndex);
+	if ( !GenerationSettings )
+		GenerationSettings = HairCardGeneratorEditor_Impl::CreateGroomSettingsLast(NewGroomAsset, CardsDesc);
+
+	bool bSuccess = false;
+	if ( !HairCardGenWindow_Utils::PromptUserWithHairCardGenDialog(GenerationSettings) )
+	{
+		return false;
+	}
+
+	// Use num groups and out group id to determine mesh name
+	int NumPhysGroups = NewGroomAsset->GetHairDescriptionGroups().HairGroups.Num();
+
+	// Check primary settings object merged and all group settings for differences
+	uint8 AllGenFlags = GenerationSettings->GetAllPipelineGeneratedDifferences();
+	if ( !GenerationSettings->CheckGenerationFlags(AllGenFlags, EHairCardGenerationPipeline::All) && CardsDesc.ImportedMesh != nullptr )
+	{
+		UE_LOG(LogHairCardGenerator, Display, TEXT("No settings changes, skipping hair card run"));
+		HairCardGeneratorEditor_Impl::CreateFinishNotification(true, LOCTEXT("HairCardGen.UpToDate", "Hair card assets are up to date"));
+		// FIXME: Currently returning true clears textures unless a reimport is run
+		return false;
+	}
+
+	// Force reimport if we are going to run any pipeline step
+	AllGenFlags |= (uint8)EHairCardGenerationPipeline::ImportUpdate;
+
+	// Check primary settings object (without group settings) used to force reruns for e.g. random seed changes
+	uint8 PipelineGenFlags = GenerationSettings->GetPipelineGeneratedDifferences();
+
+	// Write the settings that we will attempt to output to card desc (won't be stored on failure see FGroomRenderingDetails::OnGenerateCardDataUsingPlugin for details)
+	CardsDesc.GenerationSettings = GenerationSettings;
+
+	// TODO: On error will force full regen (is this the best we can do?)
+	GenerationSettings->ClearPipelineGenerated();
+
+	// Output generation settings to json file in Metadata folder
+	GenerationSettings->WriteGenerationSettings();
+
+	if (!HairCardGeneratorEditor_Impl::LoadGroomData(NewGroomAsset))
+	{
+		UE_LOG(LogHairCardGenerator, Error, TEXT("Failed to load groom asset. See log for details."));
+		HairCardGeneratorEditor_Impl::CreateFinishNotification(false, LOCTEXT("HairCardGen.FailedGroomLoad", "Failed to load groom asset. See log for details."));
+		return false;
+	}
+
+	GenerationSettings->UpdateStrandFilterAssignment();
+	if ( !CardGenController->LoadSettings(GenerationSettings) )
+	{
+		UE_LOG(LogHairCardGenerator, Error, TEXT("Failed to load hair card generation settings. See log for details."));
+		HairCardGeneratorEditor_Impl::CreateFinishNotification(false, LOCTEXT("HairCardGen.FailedSettingsLoad", "Failed to load hair card generation settings. See log for details."));
+		return false;
+	}
+
+	TArray <float> HairWidths, RootScales, TipScales;
+	HairWidths.Init(-1.0, NumPhysGroups);
+	RootScales.Init(1.0, NumPhysGroups);
+	TipScales.Init(1.0, NumPhysGroups);
+	if ( GenerationSettings->bUseGroomAssetStrandWidth )
+	{
+		for ( int PhysGroupIdx = 0; PhysGroupIdx<NumPhysGroups; PhysGroupIdx++ )
+		{
+			const FHairGeometrySettings& GeometrySettings = NewGroomAsset->GetHairGroupsRendering()[PhysGroupIdx].GeometrySettings;
+			if ( GeometrySettings.HairWidth_Override )
+			{
+				HairWidths[PhysGroupIdx] = GeometrySettings.HairWidth;
+			}
+			RootScales[PhysGroupIdx] = GeometrySettings.HairRootScale;
+			TipScales[PhysGroupIdx] = GeometrySettings.HairTipScale;
+		}
+	}
+	if ( !CardGenController->UpdateWidthMultipliers(HairWidths, RootScales, TipScales) )
+	{
+		UE_LOG(LogHairCardGenerator, Warning, TEXT("Strand width values inconsistent with number of physics groups, not using widths."));
+	}
+
+	// Run card group generation in loop
+	bool bAllGenerated = true;
+	for ( int i=0; i < GenerationSettings->GetFilterGroupSettings().Num(); ++i )
+	{
+		uint8 GroupGenFlags = PipelineGenFlags | GenerationSettings->GetPipelineFilterGroupDifferences(i);
+
+		// Skip if no group update flags required for current group
+		if ( !GenerationSettings->CheckGenerationFlags(GroupGenFlags, EHairCardGenerationPipeline::GroupUpdate) )
+		{
+			UE_LOG(LogHairCardGenerator, Display, TEXT("Setttings group %d already generated. Skipping."), i);
+			continue;
+		}
+
+		// TODO: On errors this will force full regeneration (because we clear the group gen object)
+		GenerationSettings->ClearPipelineGeneratedFilterGroup(i);
+		bAllGenerated &= HairCardGeneratorEditor_Impl::GenerateCardsForCardGroup(GenerationSettings, i, GroupGenFlags);
+
+		// TODO: It might be useful to handle user cancellation by breaking vs. errors by continuing here
+		if ( !bAllGenerated )
+		{
+			break;
+		}
+
+		GenerationSettings->WritePipelineGeneratedFilterGroup(i);
+	}
+
+	if ( !bAllGenerated )
+	{
+		UE_LOG(LogHairCardGenerator, Error, TEXT("Failed to generate cards for all groups. See log for details."));
+		HairCardGeneratorEditor_Impl::CreateFinishNotification(false, LOCTEXT("HairCardGen.FailedGeneration", "Failed to generate cards for all groups. See log for details."));
+		return false;
+	}
+
+	if ( GenerationSettings->CheckGenerationFlags(AllGenFlags, EHairCardGenerationPipeline::TextureLayout) && !GenerationSettings->bReduceCardsFromPreviousLOD )
+	{
+		FScopedSlowTask SlowTask(0, LOCTEXT("HairCardGen.BuildingLayout", "Building texture atlas layout..."));
+		SlowTask.MakeDialog(/*bShowCancelButton =*/true);
+		if ( !CardGenController->GenerateTextureLayout() )
+		{
+			UE_LOG(LogHairCardGenerator, Error, TEXT("Texture layout failed. See log for details."));
+			HairCardGeneratorEditor_Impl::CreateFinishNotification(false, LOCTEXT("HairCardGen.FailedLayout", "Texture layout failed. See log for details."));
+			return false;
+		}
+	}
+
+	if ( GenerationSettings->CheckGenerationFlags(AllGenFlags, EHairCardGenerationPipeline::TextureRendering) && !GenerationSettings->bReduceCardsFromPreviousLOD )
+	{
+		FScopedSlowTask SlowTask(0, LOCTEXT("HairCardGen.BuildingAtlas", "Exporting textures..."));
+		SlowTask.MakeDialog(/*bShowCancelButton =*/true);
+		if ( !CardGenController->GenerateTextureAtlases() )
+		{
+			UE_LOG(LogHairCardGenerator, Error, TEXT("Texture atlas creation failed. See log for details."));
+			HairCardGeneratorEditor_Impl::CreateFinishNotification(false, LOCTEXT("HairCardGen.FailedAtlas", "Texture atlas creation failed. See log for details."));
+			return false;
+		}
+	}
+
+	// Import Textures (NOTE: Mesh import no longer supported)
+	if ( GenerationSettings->CheckGenerationFlags(AllGenFlags, EHairCardGenerationPipeline::ImportTextures) )
+	{
+		FHairCardImportSettings ImportSettings;
+		HairCardGeneratorEditor_Impl::GetTextureImportNames(GenerationSettings, ImportSettings);
+
+		const FString TexturePath = GenerationSettings->GetTextureImportPath();
+		const FString TextureDest = GenerationSettings->GetTextureContentPath();
+		UE_LOG(LogHairCardGenerator, Display, TEXT("Importing generated cards from: %s"), *TexturePath);
+		TArray<UObject*> ImportedAssets = HairCardGeneratorEditor_Impl::ImportHairCardTextures(TexturePath, TextureDest, ImportSettings);
+		if ( ImportedAssets.IsEmpty() )
+		{
+			UE_LOG(LogHairCardGenerator, Error, TEXT("Failed to import textures. Check '%s' for the files."), *TexturePath);
+			HairCardGeneratorEditor_Impl::CreateFinishNotification(false, LOCTEXT("HairCardGen.FailedTextureImport", "Failed to import textures. See log for details."));
+			return false;
+		}
+
+		const int32 NumAssigned = HairCardGeneratorEditor_Impl::AssignImportedTextures(ImportedAssets, ImportSettings, CardsDesc);
+		if ( NumAssigned == 0 )
+		{
+			UE_LOG(LogHairCardGenerator, Error, TEXT("Unable to assign any of the imported assets to the groom. Verify texture file name format."));
+			HairCardGeneratorEditor_Impl::CreateFinishNotification(false, LOCTEXT("HairCardGen.FailedAssignImports", "Unable to assign any of the imported assets to the groom. See log for details."));
+			return false;
+		}
+	}
+
+	// Generate static mesh from cards and add to destination package
+	if ( GenerationSettings->CheckGenerationFlags(AllGenFlags, EHairCardGenerationPipeline::GenerateMesh) )
+	{
+		FScopedSlowTask SlowTask(0, LOCTEXT("GeneratingHairCards", "Building mesh..."));
+		SlowTask.MakeDialog(/*bShowCancelButton =*/true);
+
+		// Don't bother adding _PG# at end of mesh name if there's only one physics group
+		int OutGroupId = (NumPhysGroups > 1) ? GenerationSettings->GetGenerateForGroomGroup() : -1;
+		TObjectPtr<UStaticMesh> NewMesh = HairCardGeneratorEditor_Impl::CreateMeshInPackage(GenerationSettings->DestinationPath.Path, GenerationSettings->BaseFilename, OutGroupId);
+		if ( !NewMesh )
+		{
+			
+			UE_LOG(LogHairCardGenerator, Error, TEXT("Unable to create mesh in destination package."));
+			HairCardGeneratorEditor_Impl::CreateFinishNotification(false, LOCTEXT("HairCardGen.FailedObjectCreate", "Unable to create mesh in destination package. See log for details."));
+			return false;
+		}
+
+		if ( !CardGenController->GenerateMesh(NewMesh) )
+		{
+			UE_LOG(LogHairCardGenerator, Error, TEXT("Mesh generation failed. See log for details."));
+			HairCardGeneratorEditor_Impl::CreateFinishNotification(false, LOCTEXT("HairCardGen.FailedMeshGen", "Mesh generation failed. See log for details."));
+			return false;
+		}
+
+		CardsDesc.ImportedMesh = NewMesh;
+		CardsDesc.UpdateMeshKey();
+
+		NewMesh->PostEditChange();
+		NewMesh->GetPackage()->MarkPackageDirty();
+	}
+
+	GenerationSettings->WritePipelineGenerated();
+	HairCardGeneratorEditor_Impl::CreateFinishNotification(true, LOCTEXT("HairCardGen.GenSuccess", "Hair cards generated successfully."));
+	return true;
+}
+
+bool FHairCardGeneratorEditorModule::IsCompatibleSettings(UHairCardGenerationSettings* OldSettings)
+{
+	return UHairCardGeneratorPluginSettings::IsCompatibleSettings(OldSettings);
+}
+
+IMPLEMENT_MODULE(/*FDefaultModuleImpl*/FHairCardGeneratorEditorModule, HairCardGeneratorEditor);
+#undef LOCTEXT_NAMESPACE
diff --git a/Engine/Plugins/Experimental/HairCardGenerator/Source/HairCardGeneratorEditor/Private/HairCardGeneratorEditorSettings.cpp b/Engine/Plugins/Experimental/HairCardGenerator/Source/HairCardGeneratorEditor/Private/HairCardGeneratorEditorSettings.cpp
new file mode 100644
--- /dev/null
+++ b/Engine/Plugins/Experimental/HairCardGenerator/Source/HairCardGeneratorEditor/Private/HairCardGeneratorEditorSettings.cpp
@@ -0,0 +1,11 @@
+// Copyright Epic Games, Inc. All Rights Reserved.
+
+#include "HairCardGeneratorEditorSettings.h"
+
+#include UE_INLINE_GENERATED_CPP_BY_NAME(HairCardGeneratorEditorSettings)
+
+UHairCardGeneratorEditorSettings::UHairCardGeneratorEditorSettings()
+	: HairCardAssetNameFormat(TEXT("{groom_name}_LOD{lod}"))
+	, HairCardAssetPathFormat(TEXT("{groom_dir}"))
+{
+}
diff --git a/Engine/Plugins/Experimental/HairCardGenerator/Source/HairCardGeneratorEditor/Private/HairCardGeneratorEditorSettings.h b/Engine/Plugins/Experimental/HairCardGenerator/Source/HairCardGeneratorEditor/Private/HairCardGeneratorEditorSettings.h
new file mode 100644
--- /dev/null
+++ b/Engine/Plugins/Experimental/HairCardGenerator/Source/HairCardGeneratorEditor/Private/HairCardGeneratorEditorSettings.h
@@ -0,0 +1,29 @@
+// Copyright Epic Games, Inc. All Rights Reserved.
+
+#pragma once
+
+#include "CoreMinimal.h"
+#include "UObject/Object.h"
+#include "UObject/ObjectMacros.h"
+#include "HairCardGeneratorEditorSettings.generated.h"
+
+enum class EHairAtlasTextureType : uint8;
+
+
+UCLASS(config = Editor)
+class UHairCardGeneratorEditorSettings : public UObject
+{
+	GENERATED_BODY()
+
+public:
+	UHairCardGeneratorEditorSettings();
+
+	UPROPERTY(Config)
+	FString HairCardAssetNameFormat;
+
+	UPROPERTY(Config)
+	FString HairCardAssetPathFormat;
+
+	UPROPERTY(Config)
+	TMap<EHairAtlasTextureType, FString> HairCardTextureNameFormats;
+};
\ No newline at end of file
diff --git a/Engine/Plugins/Experimental/HairCardGenerator/Source/HairCardGeneratorEditor/Private/HairCardGeneratorLog.h b/Engine/Plugins/Experimental/HairCardGenerator/Source/HairCardGeneratorEditor/Private/HairCardGeneratorLog.h
new file mode 100644
--- /dev/null
+++ b/Engine/Plugins/Experimental/HairCardGenerator/Source/HairCardGeneratorEditor/Private/HairCardGeneratorLog.h
@@ -0,0 +1,7 @@
+// Copyright Epic Games, Inc. All Rights Reserved.
+
+#pragma once
+
+#include "Logging/LogMacros.h"
+
+DECLARE_LOG_CATEGORY_EXTERN(LogHairCardGenerator, Log, All)
\ No newline at end of file
diff --git a/Engine/Plugins/Experimental/HairCardGenerator/Source/HairCardGeneratorEditor/Private/HairCardGeneratorPluginSettings.cpp b/Engine/Plugins/Experimental/HairCardGenerator/Source/HairCardGeneratorEditor/Private/HairCardGeneratorPluginSettings.cpp
new file mode 100644
--- /dev/null
+++ b/Engine/Plugins/Experimental/HairCardGenerator/Source/HairCardGeneratorEditor/Private/HairCardGeneratorPluginSettings.cpp
@@ -0,0 +1,1235 @@
+// Copyright Epic Games, Inc. All Rights Reserved.
+
+#include "HairCardGeneratorPluginSettings.h"
+
+#include "HairCardGeneratorEditorSettings.h"
+#include "HairCardGeneratorLog.h"
+#include "HairCardGenCardSubdivider.h"
+
+#include "GroomAsset.h"
+#include "Containers/Array.h"
+#include "Containers/UnrealString.h"
+#include "Dom/JsonObject.h"
+#include "Engine/StaticMesh.h"
+#include "HAL/FileManager.h"
+#include "HAL/PlatformFileManager.h"
+#include "Misc/FileHelper.h"
+#include "Misc/PackageName.h"
+#include "ObjectTools.h"
+#include "JsonObjectConverter.h"
+#include "Serialization/JsonReader.h"
+#include "Serialization/JsonWriter.h"
+#include "Serialization/JsonSerializer.h"
+#include "Serialization/MemoryReader.h"
+#include "Serialization/MemoryWriter.h"
+
+
+/* Json helper functions
+ *****************************************************************************/
+namespace FJsonHelper
+{
+	// TODO: Need to cache settings with IDirectoryWatcher or FFileCache to avoid hitting fs constantly in UI (similar issue: HCS settings change detection)
+	TSharedPtr<FJsonObject> ReadJsonSettings(const FString& SettingsFilename);
+
+	FString GetStringFieldDefault(const TSharedPtr<FJsonObject>& Json, const FString& FieldName, const FString& DefaultVal = TEXT(""));
+	bool GetBoolFieldDefault(const TSharedPtr<FJsonObject>& Json, const FString& FieldName, bool DefaultVal = false);
+	int GetIntFieldDefault(const TSharedPtr<FJsonObject>& Json, const FString& FieldName, int DefaultVal = 0);
+};
+
+TSharedPtr<FJsonObject> FJsonHelper::ReadJsonSettings(const FString& SettingsFilename)
+{
+	TSharedPtr<FJsonObject> JsonSettingsObject;
+
+	FString JsonText;
+	if ( !FFileHelper::LoadFileToString(JsonText, *SettingsFilename) )
+	{
+		// UE_LOG(LogHairCardGenerator, Warning, TEXT("Unable to open json settings file: %s"), *SettingsFilename);
+		return nullptr;
+	}
+
+	TSharedRef<TJsonReader<TCHAR>> JsonReader = TJsonReaderFactory<TCHAR>::Create(JsonText);
+	if ( !FJsonSerializer::Deserialize(JsonReader.Get(), JsonSettingsObject) )
+	{
+		// UE_LOG(LogHairCardGenerator, Warning, TEXT("Unable to parse json settings file: %s"), *SettingsFilename);
+		return nullptr;
+	}
+
+	return JsonSettingsObject;
+}
+
+FString FJsonHelper::GetStringFieldDefault(const TSharedPtr<FJsonObject>& Json, const FString& FieldName, const FString& DefaultVal)
+{
+	if (!Json->HasTypedField<EJson::String>(FieldName))
+	{
+		return DefaultVal;
+	}
+
+	return Json->GetStringField(FieldName);
+}
+
+bool FJsonHelper::GetBoolFieldDefault(const TSharedPtr<FJsonObject>& Json, const FString& FieldName, bool DefaultVal)
+{
+	if (!Json->HasTypedField<EJson::Boolean>(FieldName))
+	{
+		return DefaultVal;
+	}
+
+	return Json->GetBoolField(FieldName);
+}
+
+int FJsonHelper::GetIntFieldDefault(const TSharedPtr<FJsonObject>& Json, const FString& FieldName, int DefaultVal)
+{
+	if (!Json->HasTypedField<EJson::Number>(FieldName))
+	{
+		return DefaultVal;
+	}
+
+	return Json->GetIntegerField(FieldName);
+}
+
+
+/* Hair card helpers
+ *****************************************************************************/
+namespace HairCardSettings_Helpers
+{
+    FString GetDefaultBaseNameForHairCards(TObjectPtr<const UGroomAsset> GroomAsset, int LODIndex);
+    FString GetDefaultPathForHairCards(TObjectPtr<const UGroomAsset> GroomAsset, int LODIndex);
+
+	void ResetSettingsFromTemplate(TObjectPtr<UObject> SettingsObject, TObjectPtr<const UObject> TemplateObject);
+	uint8 PipelineSettingsDiff(TObjectPtr<const UObject> SettingsObject, TObjectPtr<UObject> DiffObject);
+	bool IsWriteableSettingsProperty(FProperty* Property);
+
+    void GetStrandCardGroupInfo(TObjectPtr<const UGroomAsset> GroomAsset, TArray<FName>& OutGrpIds);
+};
+
+
+/// Get default base name for card generation (assets/directories) from current groom asset
+/// @param GroomAsset The groom asset cards will be generated for
+/// @param LODIndex The LOD index for card generation (TODO: This should be removed when multi-lod supported)
+/// @return Base name for generated data (e.g. assets and intermediate data)
+FString HairCardSettings_Helpers::GetDefaultBaseNameForHairCards(TObjectPtr<const UGroomAsset> GroomAsset, int LODIndex)
+{
+    const UHairCardGeneratorEditorSettings* ConfigSettings = GetDefault<UHairCardGeneratorEditorSettings>();
+
+    FFormatNamedArguments FormatArgs;
+    if (GroomAsset)
+    {
+        FormatArgs.Add(TEXT("groom_name"), FText::FromString(GroomAsset->GetName()));
+    }
+    FormatArgs.Add(TEXT("lod"), LODIndex);
+
+    return FText::Format(FText::FromString(ConfigSettings->HairCardAssetNameFormat), FormatArgs).ToString();
+}
+
+
+/// Get default output content path for card generation from current groom asset
+/// @param GroomAsset The groom asset cards will be generated for
+/// @param LODIndex The LOD index for card generation (TODO: This should be removed when multi-lod supported)
+/// @return Content path for generated hair card assets
+FString HairCardSettings_Helpers::GetDefaultPathForHairCards(TObjectPtr<const UGroomAsset> GroomAsset, int LODIndex)
+{
+    const UHairCardGeneratorEditorSettings* ConfigSettings = GetDefault<UHairCardGeneratorEditorSettings>();
+
+    FFormatNamedArguments FormatArgs;
+    if (GroomAsset)
+    {
+        FString PackageRoot, PackagePath, PackageName;
+        FPackageName::SplitLongPackageName(FPackageName::ObjectPathToPackageName(GroomAsset->GetPathName()),
+            PackageRoot, PackagePath, PackageName);
+
+        FormatArgs.Add(TEXT("groom_dir"), FText::FromString(PackageRoot + PackagePath));
+        FormatArgs.Add(TEXT("groom_name"), FText::FromString(GroomAsset->GetName()));
+    }
+    else
+    {
+        FormatArgs.Add(TEXT("groom_dir"), FText::FromString(FPaths::ProjectContentDir()));
+    }
+    FormatArgs.Add(TEXT("lod"), LODIndex);
+
+    FString DirPath = FText::Format(FText::FromString(ConfigSettings->HairCardAssetPathFormat), FormatArgs).ToString();
+    FPaths::NormalizeDirectoryName(DirPath);
+    FPaths::RemoveDuplicateSlashes(DirPath);
+
+    return DirPath;
+}
+
+void HairCardSettings_Helpers::ResetSettingsFromTemplate(TObjectPtr<UObject> SettingsObject, TObjectPtr<const UObject> TemplateObject)
+{
+	if ( !TemplateObject->IsA(SettingsObject->GetClass()) )
+	{
+		UE_LOG(LogHairCardGenerator, Warning, TEXT("Invalid template object for copying settings (%s)"), *SettingsObject->GetName());
+		return;
+	}
+
+	for (TFieldIterator<FProperty> PropertyIt(SettingsObject->GetClass(), EFieldIterationFlags::IncludeDeprecated); PropertyIt; ++PropertyIt)
+	{
+		if ( !HairCardSettings_Helpers::IsWriteableSettingsProperty(*PropertyIt) )
+			continue;
+
+		const void* SrcAddress = PropertyIt->ContainerPtrToValuePtr<void>(TemplateObject);
+		void* DstAddress = PropertyIt->ContainerPtrToValuePtr<void>(SettingsObject);
+
+		PropertyIt->CopyCompleteValue(DstAddress, SrcAddress);
+	}
+}
+
+uint8 HairCardSettings_Helpers::PipelineSettingsDiff(TObjectPtr<const UObject> SettingsObject, TObjectPtr<UObject> DiffObject)
+{
+	if ( !DiffObject->IsA(SettingsObject->GetClass()) )
+	{
+		UE_LOG(LogHairCardGenerator, Warning, TEXT("Comparing settings objects of incompatible types (%s,%s)"), *SettingsObject->GetClass()->GetName(), *DiffObject->GetClass()->GetName());
+		return (uint8)(EHairCardGenerationSettingsCategories::All);
+	}
+
+	uint8 PipelineDiff = (uint8)EHairCardGenerationSettingsCategories::None;
+
+	const UEnum* HairCardSettingsCategoriesEnum = StaticEnum<EHairCardGenerationSettingsCategories>();
+	for (TFieldIterator<FProperty> PropertyIt(SettingsObject->GetClass(), EFieldIterationFlags::IncludeDeprecated); PropertyIt; ++PropertyIt)
+	{
+		if ( !HairCardSettings_Helpers::IsWriteableSettingsProperty(*PropertyIt) )
+			continue;
+
+		if ( !PropertyIt->HasMetaData(TEXT("Category")) )
+		{
+			UE_LOG(LogHairCardGenerator, Warning, TEXT("Hair Card Settings field: %s has no category metadata (required for change detection)!"), *PropertyIt->GetName());
+			continue;
+		}
+
+		FString BaseCategoryName = PropertyIt->GetMetaData(TEXT("Category"));
+
+		// We use category names to determine pipeline effects of a difference
+		BaseCategoryName.RemoveSpacesInline();
+		int32 SubCatIdx = BaseCategoryName.Find(TEXT("|"));
+		if ( SubCatIdx != INDEX_NONE )
+			BaseCategoryName.LeftInline(SubCatIdx);
+
+		int64 CategoryPipelineDiff = HairCardSettingsCategoriesEnum->GetValueByNameString(BaseCategoryName);
+		if ( CategoryPipelineDiff == INDEX_NONE )
+		{
+			UE_LOG(LogHairCardGenerator, Warning, TEXT("Hair Card Settings field: %s has unexpected category \"%s\" (cannot detect pipeline change effects)!"), *PropertyIt->GetName(), *BaseCategoryName);
+			continue;
+		}
+
+		const void* SettingsValueAddr = PropertyIt->ContainerPtrToValuePtr<void>(SettingsObject);
+		const void* DiffValueAddr = PropertyIt->ContainerPtrToValuePtr<void>(DiffObject);
+
+		if ( !PropertyIt->Identical(SettingsValueAddr, DiffValueAddr) )
+			PipelineDiff |= CategoryPipelineDiff;
+	}
+
+	return PipelineDiff;
+}
+
+bool HairCardSettings_Helpers::IsWriteableSettingsProperty(FProperty* Property)
+{
+	// Determine if this is a settings property that should be written/compared
+	// Ignore UObject properties that are private, protected, or transient
+	if ( Property->HasAnyPropertyFlags(CPF_NativeAccessSpecifierPrivate | CPF_NativeAccessSpecifierProtected | CPF_Transient) )
+	{
+		return false;
+	}
+	
+	// Property must be BlueprintVisible (accessible from Python)
+	return Property->HasAllPropertyFlags(CPF_BlueprintVisible);
+}
+
+void HairCardSettings_Helpers::GetStrandCardGroupInfo(TObjectPtr<const UGroomAsset> GroomAsset, TArray<FName>& OutGrpIds)
+{
+    TSet<FName> IdNames;
+
+    FHairDescription HairDesc = GroomAsset->GetHairDescription();
+	FStrandAttributesRange TraverseStrands(HairDesc, FHairStrandAttribsConstIterator::IgnoreNone);
+    if ( TraverseStrands.First().HasCardGroupName() )
+    {
+		for ( const FHairStrandAttribsRefProxy& StrandAttribs : TraverseStrands )
+		{
+			IdNames.Add(StrandAttribs.GetCardGroupName());
+		}
+
+        // IdCountMap.KeySort([](FName a, FName b){return a < b;});
+    }
+
+	if ( IdNames.Contains(NAME_None) )
+	{
+		IdNames.Remove(NAME_None);
+		UE_LOG(LogHairCardGenerator, Warning, TEXT("Groom has groom_group_cards_id attribute, but some strands are empty (will be ignored)"));
+	}
+
+    OutGrpIds = IdNames.Array();
+}
+
+
+/* HairCardSettings_Converter
+ *****************************************************************************/
+namespace HairCardSettings_Converter
+{
+	bool ResetSettingsFromJson(UObject* SettingsObject, const TSharedPtr<FJsonObject>& JsonObject);
+	bool WriteSettingsToJson(UObject* SettingsObject, TSharedRef<FJsonObject> JsonRoot);
+}
+
+bool HairCardSettings_Converter::ResetSettingsFromJson(UObject* SettingsObject, const TSharedPtr<FJsonObject>& JsonObject)
+{
+	for (TFieldIterator<FProperty> PropertyIt(SettingsObject->GetClass(), EFieldIterationFlags::IncludeDeprecated); PropertyIt; ++PropertyIt)
+	{
+		if ( !HairCardSettings_Helpers::IsWriteableSettingsProperty(*PropertyIt) )
+			continue;
+
+		TSharedPtr<FJsonValue> PropValue = JsonObject->TryGetField(PropertyIt->GetName());
+		if ( !PropValue )
+		{
+			UE_LOG(LogHairCardGenerator, Warning, TEXT("Json Field Not Found: %s"), *PropertyIt->GetName());
+			continue;
+		}
+
+		void* SettingsPropPtr = PropertyIt->ContainerPtrToValuePtr<void>(SettingsObject);
+		FJsonObjectConverter::JsonValueToUProperty(PropValue, *PropertyIt, SettingsPropPtr);
+	}
+
+	return true;
+}
+
+bool HairCardSettings_Converter::WriteSettingsToJson(UObject* SettingsObject, TSharedRef<FJsonObject> JsonRoot)
+{
+	// Convert all blueprint-accessible non-transient/private/protected properties
+	for (TFieldIterator<FProperty> PropertyIt(SettingsObject->GetClass(), EFieldIterationFlags::IncludeDeprecated); PropertyIt; ++PropertyIt)
+	{
+		if ( !HairCardSettings_Helpers::IsWriteableSettingsProperty(*PropertyIt) )
+			continue;
+
+		void* SettingsPropPtr = PropertyIt->ContainerPtrToValuePtr<void>(SettingsObject);
+		TSharedPtr<FJsonValue> JsonValue = FJsonObjectConverter::UPropertyToJsonValue(*PropertyIt, SettingsPropPtr);
+		JsonRoot->SetField(PropertyIt->GetName(), JsonValue);
+	}
+
+	return true;
+}
+
+
+
+/* UHairCardGeneratorPluginSettings
+ *****************************************************************************/
+
+void UHairCardGeneratorPluginSettings::SetSource(TObjectPtr<const UGroomAsset> InSourceObject, int32 GenLODIndex, int32 PhysGroupIndex)
+{
+	GroomAsset = InSourceObject;
+	LODIndex = GenLODIndex;
+	GenerateForGroomGroup = PhysGroupIndex;
+}
+
+void UHairCardGeneratorPluginSettings::ResetFromSourceDescription(const FHairGroupsCardsSourceDescription& SourceDesc)
+{
+	TObjectPtr<const UHairCardGeneratorPluginSettings> LastSettings = Cast<UHairCardGeneratorPluginSettings>(SourceDesc.GenerationSettings);
+
+	LODIndex = SourceDesc.LODIndex;
+	GenerateForGroomGroup = SourceDesc.GroupIndex;
+	OldGeneratedMesh = SourceDesc.ImportedMesh;
+
+	if ( LastSettings == nullptr )
+		LastSettings = UHairCardGeneratorPluginSettings::StaticClass()->GetDefaultObject<UHairCardGeneratorPluginSettings>();
+
+	HairCardSettings_Helpers::ResetSettingsFromTemplate(this, LastSettings);
+	if (OldGeneratedMesh && LODIndex == LastSettings->LODIndex)
+	{
+		BaseFilename = OldGeneratedMesh->GetName();
+		DestinationPath.Path = FPackageName::GetLongPackagePath(OldGeneratedMesh->GetOuter()->GetPathName());
+	}
+	else
+	{
+		BaseFilename = HairCardSettings_Helpers::GetDefaultBaseNameForHairCards(GroomAsset, LODIndex);
+		DestinationPath.Path = HairCardSettings_Helpers::GetDefaultPathForHairCards(GroomAsset, LODIndex);
+	}
+
+	HairCardSettings_Helpers::GetStrandCardGroupInfo(GroomAsset, CardGroupIds);
+
+	int GroupSettingsCount = FMath::Max(1, LastSettings->FilterGroupGenerationSettings.Num());
+	for ( int i=0; i < GroupSettingsCount; ++i )
+	{
+		TObjectPtr<UHairCardGeneratorGroupSettings> oldOpts;
+		if ( LastSettings->FilterGroupGenerationSettings.Num() > i )
+			oldOpts = LastSettings->FilterGroupGenerationSettings[i];
+		else
+		{
+			oldOpts = UHairCardGeneratorGroupSettings::StaticClass()->GetDefaultObject<UHairCardGeneratorGroupSettings>();
+		}
+
+		ResetFilterGroupSettings(i, oldOpts);
+	}
+	ResetNumFilterGroups(GroupSettingsCount);
+
+	UpdateOutputPaths();
+	UpdateParentInfo();
+}
+
+bool UHairCardGeneratorPluginSettings::ResetFromSettingsJson()
+{
+	// HACK: Need to bootstrap base filename and intermediate path in order to get json path
+	if ( BaseFilename.IsEmpty() )
+	{
+		BaseFilename = HairCardSettings_Helpers::GetDefaultBaseNameForHairCards(GroomAsset, LODIndex);
+		DestinationPath.Path = HairCardSettings_Helpers::GetDefaultPathForHairCards(GroomAsset, LODIndex);
+		UpdateOutputPaths();
+		UpdateParentInfo();
+	}
+
+	const FString SettingsFile = GetSettingsFilename();
+	if ( !FPaths::FileExists(SettingsFile) )
+		return false;
+
+	ResetFromJson();
+	UpdateOutputPaths();
+	UpdateParentInfo();
+
+	return true;
+}
+
+
+void UHairCardGeneratorPluginSettings::ResetToDefault()
+{
+	TObjectPtr<UHairCardGeneratorPluginSettings> DefaultSettings = UHairCardGeneratorPluginSettings::StaticClass()->GetDefaultObject<UHairCardGeneratorPluginSettings>();
+	HairCardSettings_Helpers::ResetSettingsFromTemplate(this, DefaultSettings);
+
+	BaseFilename = HairCardSettings_Helpers::GetDefaultBaseNameForHairCards(GroomAsset, LODIndex);
+	DestinationPath.Path = HairCardSettings_Helpers::GetDefaultPathForHairCards(GroomAsset, LODIndex);
+
+	HairCardSettings_Helpers::GetStrandCardGroupInfo(GroomAsset, CardGroupIds);
+
+	ResetFilterGroupSettings(0, UHairCardGeneratorGroupSettings::StaticClass()->GetDefaultObject<UHairCardGeneratorGroupSettings>());
+	ResetNumFilterGroups(1);
+
+	UpdateOutputPaths();
+	UpdateParentInfo();
+}
+
+void UHairCardGeneratorPluginSettings::ResetFilterGroupSettingsToDefault(int FilterGroupIndex)
+{
+	TObjectPtr<UHairCardGeneratorGroupSettings> DefGrpSettings = UHairCardGeneratorGroupSettings::StaticClass()->GetDefaultObject<UHairCardGeneratorGroupSettings>();
+	ResetFilterGroupSettings(FilterGroupIndex, DefGrpSettings);
+}
+
+void UHairCardGeneratorPluginSettings::AddNewSettingsFilterGroup()
+{
+	// TODO: Remove this check when more general filter support is added to UI
+	if ( FilterGroupGenerationSettings.Num() >= CardGroupIds.Num() )
+		return;
+
+	// Add a new default settings group
+	ResetFilterGroupSettings(FilterGroupGenerationSettings.Num(), UHairCardGeneratorGroupSettings::StaticClass()->GetDefaultObject<UHairCardGeneratorGroupSettings>());
+}
+
+void UHairCardGeneratorPluginSettings::RemoveSettingsFilterGroup(TObjectPtr<UHairCardGeneratorGroupSettings> SettingsGroup)
+{
+	// Always keep one settings group
+	if ( FilterGroupGenerationSettings.Num() <= 1 )
+		return;
+
+	int RemoveIndex = FilterGroupGenerationSettings.Find(SettingsGroup);
+	if ( RemoveIndex != INDEX_NONE )
+		FilterGroupGenerationSettings.RemoveAt(RemoveIndex);
+}
+
+void UHairCardGeneratorPluginSettings::WriteGenerationSettings()
+{
+	const FString SettingsPath = GetMetadataPath();
+
+	IPlatformFile& PlatformFile = FPlatformFileManager::Get().GetPlatformFile();
+	PlatformFile.CreateDirectoryTree(*SettingsPath);
+
+	WriteToJson();
+}
+
+bool UHairCardGeneratorPluginSettings::CheckGenerationFlags(uint8 NeedsGenFlags, const EHairCardGenerationPipeline PipelineStage)
+{
+	return ((NeedsGenFlags & (uint8)PipelineStage) != 0);
+}
+
+uint8 UHairCardGeneratorPluginSettings::GetAllPipelineGeneratedDifferences() const
+{
+	uint8 PipelineFlags = GetPipelineGeneratedDifferences();
+	if (PipelineFlags == (uint8)EHairCardGenerationSettingsCategories::All)
+		return PipelineFlags;
+
+	for ( int i=0; i < FilterGroupGenerationSettings.Num(); ++i )
+		PipelineFlags |= FilterGroupGenerationSettings[i]->GetPipelineFilterGroupDifferences();
+
+	return PipelineFlags;
+}
+
+uint8 UHairCardGeneratorPluginSettings::GetPipelineGeneratedDifferences() const
+{
+	const FString SettingsBinFile = GetGeneratedSettingsFilename();
+	if ( !FPaths::FileExists(SettingsBinFile) )
+		return (uint8)EHairCardGenerationPipeline::All;
+
+	TArray<uint8> SettingsBytes;
+	if ( !FFileHelper::LoadFileToArray(SettingsBytes, *SettingsBinFile, 0) )
+		return (uint8)EHairCardGenerationPipeline::All;
+
+	FMemoryReader MemReader(SettingsBytes, true);
+
+	TObjectPtr<UHairCardGeneratorPluginSettings> DiffObject = NewObject<UHairCardGeneratorPluginSettings>();
+	DiffObject->SerializeEditableSettings(MemReader);
+	DiffObject->BaseFilename = BaseFilename;
+
+	// Early-out if HCS tool versions differ from generated settings
+	if (DiffObject->Version != Version)
+		return (uint8)EHairCardGenerationSettingsCategories::All;
+
+	return HairCardSettings_Helpers::PipelineSettingsDiff(this, DiffObject);
+}
+
+void UHairCardGeneratorPluginSettings::ClearPipelineGenerated() const
+{
+	const FString SettingsBinFile = GetGeneratedSettingsFilename();
+
+	IPlatformFile& FileManager = FPlatformFileManager::Get().GetPlatformFile();
+	FileManager.DeleteFile(*SettingsBinFile);
+}
+
+void UHairCardGeneratorPluginSettings::WritePipelineGenerated()
+{
+	const FString GeneratedPath = GetGeneratedSettingsPath();
+	IPlatformFile& FileManager = FPlatformFileManager::Get().GetPlatformFile();
+	if ( !FileManager.DirectoryExists(*GeneratedPath) )
+		FileManager.CreateDirectoryTree(*GeneratedPath);
+
+	TArray<uint8> SettingsBytes;
+	SettingsBytes.Reserve(64);
+	FMemoryWriter MemWriter(SettingsBytes, true);
+
+	SerializeEditableSettings(MemWriter);
+
+	const FString SettingsBinFile = GetGeneratedSettingsFilename();
+	FFileHelper::SaveArrayToFile(SettingsBytes, *SettingsBinFile);
+}
+
+uint8 UHairCardGeneratorPluginSettings::GetPipelineFilterGroupDifferences(int FilterGroupIndex) const
+{
+	if ( FilterGroupIndex > FilterGroupGenerationSettings.Num() )
+		return (uint8)EHairCardGenerationPipeline::All;
+
+	return FilterGroupGenerationSettings[FilterGroupIndex]->GetPipelineFilterGroupDifferences();
+}
+
+void UHairCardGeneratorPluginSettings::ClearPipelineGeneratedFilterGroup(int FilterGroupIndex) const
+{
+	if ( FilterGroupIndex > FilterGroupGenerationSettings.Num() )
+		return;
+
+	FilterGroupGenerationSettings[FilterGroupIndex]->ClearPipelineGeneratedFilterGroup();
+}
+
+void UHairCardGeneratorPluginSettings::WritePipelineGeneratedFilterGroup(int FilterGroupIndex)
+{
+	if ( FilterGroupIndex > FilterGroupGenerationSettings.Num() )
+		return;
+
+	FilterGroupGenerationSettings[FilterGroupIndex]->WritePipelineGeneratedFilterGroup();
+}
+
+void UHairCardGeneratorPluginSettings::UpdateOutputPaths()
+{
+	IntermediatePath = GetIntermediatePath();
+	OutputPath = GetIntermediatePath() / TEXT("Output") / BaseFilename;
+	for ( int index = 0; index < FilterGroupGenerationSettings.Num(); ++index )
+		FilterGroupGenerationSettings[index]->UpdateGenerateFilename(BaseFilename, index);
+}
+
+void UHairCardGeneratorPluginSettings::UpdateParentInfo()
+{
+	BaseParentName = TEXT("");
+	// Always update parent names for LOD > 0 since may need to pull reserve texture space
+	if ( LODIndex >= 1 )
+	{
+		// TODO: Potentially allow directly specifying derived name
+		FString CurLOD = FString::Printf(TEXT("_LOD%d"), LODIndex);
+		FString PrevLOD = FString::Printf(TEXT("_LOD%d"), LODIndex-1);
+		BaseParentName = BaseFilename.Replace(*CurLOD,*PrevLOD, ESearchCase::CaseSensitive);
+	}
+
+	FindDerivedTextureSettings();
+
+	for ( int index = 0; index < FilterGroupGenerationSettings.Num(); ++index )
+		FilterGroupGenerationSettings[index]->UpdateParentName(BaseParentName, index);
+}
+
+bool UHairCardGeneratorPluginSettings::FindDerivedTextureSettings()
+{
+	DerivedTextureSettingsName = TEXT("");
+	bool bUsePreviousTexture = bReduceCardsFromPreviousLOD || bUseReservedSpaceFromPreviousLOD;
+	if ( !bUsePreviousTexture )
+	{
+		DerivedTextureSettingsName = BaseFilename;
+		return true;
+	}
+
+	TSharedPtr<FJsonObject> ParentTextureSettings = GetParentTextureSettings();
+	if ( !ParentTextureSettings.IsValid() )
+	{
+		return false;
+	}
+
+	DerivedTextureSettingsName = FJsonHelper::GetStringFieldDefault(ParentTextureSettings, TEXT("BaseFilename"));
+	return true;
+}
+
+TSharedPtr<FJsonObject> UHairCardGeneratorPluginSettings::GetParentTextureSettings() const
+{
+	return TraverseSettingsParents(BaseParentName, 
+		[](const TSharedPtr<FJsonObject>& SettingsRoot)
+		{
+			bool bDerivedGeometry = FJsonHelper::GetBoolFieldDefault(SettingsRoot, TEXT("bReduceCardsFromPreviousLOD"));
+			bool bDerivedTexture = FJsonHelper::GetBoolFieldDefault(SettingsRoot, TEXT("bUseReservedSpaceFromPreviousLOD"));
+			return !bDerivedGeometry && !bDerivedTexture;
+		}
+	);
+}
+
+TSharedPtr<FJsonObject> UHairCardGeneratorPluginSettings::TraverseSettingsParents(const FString& SettingsName, TFunction<bool(const TSharedPtr<FJsonObject>&)> StopCond)
+{
+	TSet<FString> CycleBreak;
+	FString PrevLOD = SettingsName;
+
+	// Break infinite loops with a max lod
+	const int MaxLOD = 15;
+	for (int i=0; i < MaxLOD && !PrevLOD.IsEmpty(); ++i)
+	{
+		const FString SettingsFilename = GetSettingsFileFromBase(PrevLOD);
+		TSharedPtr<FJsonObject> SettingsRoot = FJsonHelper::ReadJsonSettings(SettingsFilename);
+		if ( !SettingsRoot.IsValid() )
+		{
+			return nullptr;
+		}
+
+		// Check if this setting meets stop criterion
+		if ( StopCond(SettingsRoot) )
+		{
+			return SettingsRoot;
+		}
+
+		// Track settings names in case there's a cycle
+		if ( CycleBreak.Contains(PrevLOD) )
+		{
+			return nullptr;
+		}
+
+		CycleBreak.Add(PrevLOD);
+		PrevLOD = FJsonHelper::GetStringFieldDefault(SettingsRoot, TEXT("PreviousLOD"));
+	}
+
+	return nullptr;
+}
+
+
+FString UHairCardGeneratorPluginSettings::GetIntermediatePath()
+{
+	return FPaths::ProjectIntermediateDir() / TEXT("GroomHairCardGen");
+}
+
+FString UHairCardGeneratorPluginSettings::GetMetadataPath()
+{
+	return GetIntermediatePath() / TEXT("Metadata");
+}
+
+FString UHairCardGeneratorPluginSettings::GetGeneratedSettingsPath()
+{
+	return GetMetadataPath() / TEXT("_generated");
+}
+
+
+FString UHairCardGeneratorPluginSettings::GetTextureImportBaseName() const
+{
+	return DerivedTextureSettingsName;
+}
+
+FString UHairCardGeneratorPluginSettings::GetTextureImportPath() const
+{
+	return GetIntermediatePath() / TEXT("Output") / DerivedTextureSettingsName;
+}
+
+FString UHairCardGeneratorPluginSettings::GetTextureContentPath() const
+{
+	const FString SettingsFile = GetSettingsFileFromBase(DerivedTextureSettingsName);
+	TSharedPtr<FJsonObject> SettingsObj = FJsonHelper::ReadJsonSettings(SettingsFile);
+	return FJsonHelper::GetStringFieldDefault(SettingsObj, TEXT("DestinationPath"));
+}
+
+
+bool UHairCardGeneratorPluginSettings::HasValidFullParent() const
+{
+	TSharedPtr<FJsonObject> DerivedFullSettings = TraverseSettingsParents(BaseParentName, 
+		[](const TSharedPtr<FJsonObject>& SettingsRoot)
+		{
+			bool bDerivedGeometry = FJsonHelper::GetBoolFieldDefault(SettingsRoot, TEXT("bReduceCardsFromPreviousLOD"));
+			return !bDerivedGeometry;
+		}
+	);
+
+	return DerivedFullSettings.IsValid();
+}
+
+bool UHairCardGeneratorPluginSettings::HasDerivedTextureSettings() const
+{
+	// TODO: Cache the settings hierarchy for these checks
+	TSharedPtr<FJsonObject> ParentTextureSettings = GetParentTextureSettings();
+	return ParentTextureSettings.IsValid();
+}
+
+int UHairCardGeneratorPluginSettings::GetDerivedReservedTextureSize() const
+{
+	TSharedPtr<FJsonObject> ParentTextureSettings = GetParentTextureSettings();
+	if ( !ParentTextureSettings.IsValid() )
+	{
+		return 0;
+	}
+
+	return FJsonHelper::GetIntFieldDefault(ParentTextureSettings, TEXT("ReserveTextureSpaceLOD"));
+}
+
+
+FString UHairCardGeneratorPluginSettings::GetSettingsFileFromBase(const FString& SettingsName)
+{
+	return GetMetadataPath() / SettingsName + TEXT(".json");
+}
+
+FString UHairCardGeneratorPluginSettings::GetSettingsFilename() const
+{
+	return GetSettingsFileFromBase(BaseFilename);
+}
+
+FString UHairCardGeneratorPluginSettings::GetGeneratedSettingsFilename() const
+{
+	return GetGeneratedSettingsPath() / BaseFilename + TEXT(".settings.bin");
+}
+
+FString UHairCardGeneratorPluginSettings::GetGeneratedGroupSettingsFilename(int FilterGroupIndex) const
+{
+	return GetGeneratedSettingsPath() / FilterGroupGenerationSettings[FilterGroupIndex]->GenerateFilename + TEXT(".settings.bin");
+}
+
+void UHairCardGeneratorPluginSettings::ResetFromJson()
+{
+	FString SettingsFile = GetSettingsFilename();
+	TSharedPtr<FJsonObject> JsonSettingsObject = FJsonHelper::ReadJsonSettings(SettingsFile);
+
+	// Keep this base filename even if json is different (e.g. copied json without changing base)
+	FString TmpBase = BaseFilename;
+
+	// Load settings from json object 
+	HairCardSettings_Converter::ResetSettingsFromJson(this, JsonSettingsObject);
+	BaseFilename = TmpBase;
+
+	// Load destination path from json
+	const FString DestinationPathField = TEXT("DestinationPath");
+	if (JsonSettingsObject->HasTypedField<EJson::String>(DestinationPathField))
+	{
+		DestinationPath.Path = JsonSettingsObject->GetStringField(DestinationPathField);
+	}
+	
+	const FString GroupsArrayFieldName = TEXT("Groups");
+	// Load each group settings from json subobjects array
+	if ( !JsonSettingsObject->HasTypedField<EJson::Array>(GroupsArrayFieldName) )
+	{
+		UE_LOG(LogHairCardGenerator, Warning, TEXT("No haircard group settings found in file: %s"), *SettingsFile);
+		return;
+	}
+
+	const TArray<TSharedPtr<FJsonValue>>& JsonGroups = JsonSettingsObject->GetArrayField(GroupsArrayFieldName);
+	for ( int Index = 0; Index < JsonGroups.Num(); ++Index )
+	{
+		if ( JsonGroups[Index]->Type != EJson::Object )
+		{
+			UE_LOG(LogHairCardGenerator, Warning, TEXT("Invalid settings at index %d: %s"), Index, *SettingsFile);
+			continue;
+		}
+
+		// Load subsettings objects from json
+		const TSharedPtr<FJsonObject>& GroupSettings = JsonGroups[Index]->AsObject();
+		ResetFilterGroupSettings(Index, GroupSettings);
+	}
+	ResetNumFilterGroups(JsonGroups.Num());
+}
+
+
+void UHairCardGeneratorPluginSettings::WriteToJson()
+{
+	TSharedRef<FJsonObject> JsonSettingsRoot = MakeShared<FJsonObject>();
+	HairCardSettings_Converter::WriteSettingsToJson(this, JsonSettingsRoot);
+	// Manually output some necessary settings to json
+	JsonSettingsRoot->SetField(TEXT("PreviousLOD"), MakeShared<FJsonValueString>(BaseParentName));
+	JsonSettingsRoot->SetField(TEXT("DestinationPath"), MakeShared<FJsonValueString>(DestinationPath.Path));
+	JsonSettingsRoot->SetField(TEXT("Version"), MakeShared<FJsonValueString>(Version));
+
+	// Create array of json objects for group settings subobjects
+	TArray<TSharedPtr<FJsonValue>> JsonGroups;
+	for ( int Index = 0; Index < FilterGroupGenerationSettings.Num(); ++Index )
+	{
+		TSharedRef<FJsonObject> JsonSettingsGroup = MakeShared<FJsonObject>();
+		HairCardSettings_Converter::WriteSettingsToJson(FilterGroupGenerationSettings[Index], JsonSettingsGroup);
+
+		JsonGroups.Add(MakeShared<FJsonValueObject>(JsonSettingsGroup));
+	}
+	JsonSettingsRoot->SetArrayField(TEXT("Groups"), JsonGroups);
+
+	FString SettingsFile = GetSettingsFilename();
+
+	FString JsonText;
+	TSharedRef<TJsonWriter<TCHAR>> JsonWriter = TJsonWriterFactory<TCHAR>::Create(&JsonText);
+	if ( !FJsonSerializer::Serialize(JsonSettingsRoot, JsonWriter, true) )
+	{
+		UE_LOG(LogHairCardGenerator, Warning, TEXT("Unable to write json settings string for object"));
+		return;
+	}
+
+	if ( !FFileHelper::SaveStringToFile(JsonText, *SettingsFile) )
+	{
+		UE_LOG(LogHairCardGenerator, Warning, TEXT("Unable to write json settings file: %s"), *SettingsFile);
+		return;
+	}
+}
+
+void UHairCardGeneratorPluginSettings::SerializeEditableSettings(FArchive& Ar)
+{
+	// NOTE: Only append new settings fields at end (FArchive produces reasonable defaults without buffer overflow)
+	Ar << DestinationPath.Path;
+	Ar << RandomSeed;
+
+	Ar << AtlasSize;
+	Ar << DepthMinimum;
+	Ar << DepthMaximum;
+	Ar << UVSmoothingIterations;
+	
+	Ar << Version;
+
+	Ar << bReduceCardsFromPreviousLOD;
+	Ar << bUseReservedSpaceFromPreviousLOD;
+	Ar << ReserveTextureSpaceLOD;
+	Ar << bUseGroomAssetStrandWidth;
+}
+
+void UHairCardGeneratorPluginSettings::ResetNumFilterGroups(int Count)
+{
+	FilterGroupGenerationSettings.SetNum(Count, EAllowShrinking::Yes);
+}
+
+void UHairCardGeneratorPluginSettings::ResetFilterGroupSettings(int FilterGroupIndex, TObjectPtr<UHairCardGeneratorGroupSettings> OptionsTemplate)
+{
+	if ( FilterGroupGenerationSettings.Num() <= FilterGroupIndex )
+		FilterGroupGenerationSettings.Add(NewObject<UHairCardGeneratorGroupSettings>(this));
+
+	FilterGroupGenerationSettings[FilterGroupIndex]->SetSource(GroomAsset, LODIndex);
+    FilterGroupGenerationSettings[FilterGroupIndex]->Reset(OptionsTemplate);
+
+	FilterGroupGenerationSettings[FilterGroupIndex]->ResetFilterDefault(CardGroupIds);
+	FilterGroupGenerationSettings[FilterGroupIndex]->UpdateStrandFilter();
+
+	FilterGroupGenerationSettings[FilterGroupIndex]->UpdateGenerateFilename(BaseFilename, FilterGroupIndex);
+}
+
+void UHairCardGeneratorPluginSettings::ResetFilterGroupSettings(int FilterGroupIndex, const TSharedPtr<FJsonObject>& OptionsTemplate)
+{
+	if ( FilterGroupGenerationSettings.Num() <= FilterGroupIndex )
+		FilterGroupGenerationSettings.Add(NewObject<UHairCardGeneratorGroupSettings>(this));
+
+	FilterGroupGenerationSettings[FilterGroupIndex]->SetSource(GroomAsset, LODIndex);
+    FilterGroupGenerationSettings[FilterGroupIndex]->Reset(OptionsTemplate);
+
+	FilterGroupGenerationSettings[FilterGroupIndex]->ResetFilterDefault(CardGroupIds);
+	FilterGroupGenerationSettings[FilterGroupIndex]->UpdateStrandFilter();
+
+	FilterGroupGenerationSettings[FilterGroupIndex]->UpdateGenerateFilename(BaseFilename, FilterGroupIndex);
+}
+
+bool UHairCardGeneratorPluginSettings::ValidateStrandFilters(TMap<uint32,uint32>& ErrorCountMap)
+{
+	bool NeedsUpdate = false;
+	for ( TObjectPtr<UHairCardGeneratorGroupSettings> FilterGroupSettings : FilterGroupGenerationSettings )
+	{
+		if ( FilterGroupSettings->StrandFilterChecksDirty )
+		{
+			NeedsUpdate = true;
+			break;
+		}
+	}
+
+	if ( NeedsUpdate )
+	{
+		ValidateStrandFiltersInternal();
+	}
+
+	for ( TObjectPtr<UHairCardGeneratorGroupSettings> FilterGroupSettings : FilterGroupGenerationSettings )
+	{
+		FilterGroupSettings->StrandFilterChecksDirty = false;
+	}
+
+	ErrorCountMap = StrandErrorCountMap;
+	return (UnassignedStrandsCount > 0);
+}
+
+void UHairCardGeneratorPluginSettings::UpdateStrandFilterAssignment()
+{
+	FHairDescription HairDesc = GroomAsset->GetHairDescription();
+	StrandFilterGroupIndexMap.Reset(HairDesc.GetNumStrands());
+
+	for ( const FHairStrandAttribsRefProxy& StrandAttribs : FStrandAttributesRange(HairDesc, FHairStrandAttribsConstIterator::DefaultIgnore) )
+	{
+		StrandFilterGroupIndexMap.Add(AssignStrandToFilterGroup(StrandAttribs));
+	}
+
+	StrandFilterGroupIndexMap.Shrink();
+}
+
+void UHairCardGeneratorPluginSettings::ValidateStrandFiltersInternal()
+{
+	UnassignedStrandsCount = 0;
+	StrandErrorCountMap.Reset();
+
+	FHairDescription HairDesc = GroomAsset->GetHairDescription();
+	for ( const FHairStrandAttribsRefProxy& StrandAttribs : FStrandAttributesRange(HairDesc, FHairStrandAttribsConstIterator::DefaultIgnore) )
+	{
+		uint32 AssignmentBits = CheckAllStrandAssignments(StrandAttribs);
+		// Check for a single 1 (power of 2)
+		bool AssignedMultipleGroups = (AssignmentBits & (AssignmentBits - 1)) != 0;
+		if ( AssignmentBits == 0 )
+			UnassignedStrandsCount += 1;
+		else if ( AssignedMultipleGroups )
+			StrandErrorCountMap.FindOrAdd(AssignmentBits, 0) += 1;
+	}
+}
+
+int32 UHairCardGeneratorPluginSettings::AssignStrandToFilterGroup(const FHairStrandAttribsRefProxy& StrandAttribs)
+{
+	for ( int Index = 0; Index < FilterGroupGenerationSettings.Num(); ++Index )
+	{
+		bool InGroup = FilterGroupGenerationSettings[Index]->StrandsFilter->FilterStrand(StrandAttribs);
+		if ( InGroup )
+			return Index;
+	}
+
+	return INDEX_NONE;
+}
+
+uint32 UHairCardGeneratorPluginSettings::CheckAllStrandAssignments(const FHairStrandAttribsRefProxy& StrandAttribs)
+{
+	uint32 Assignment = 0;
+	for ( int Index = 0; Index < FilterGroupGenerationSettings.Num(); ++Index )
+	{
+		if ( FilterGroupGenerationSettings[Index]->StrandsFilter->FilterStrand(StrandAttribs) )
+			Assignment |= (1 << Index);
+	}
+
+	return Assignment;
+}
+
+bool UHairCardGeneratorPluginSettings::CanEditChange(const FProperty* InProperty) const
+{
+	bool bCanEdit = Super::CanEditChange(InProperty);
+	if (!bCanEdit)
+	{
+		return false;
+	}
+
+	FName PropertyName = InProperty->GetFName();
+
+	// Disable most settings if Reducing from previous LOD
+	if ( bReduceCardsFromPreviousLOD )
+	{
+		return !((PropertyName == GET_MEMBER_NAME_CHECKED(UHairCardGeneratorPluginSettings, bGenerateGeometryForAllGroups))
+			|| (PropertyName == GET_MEMBER_NAME_CHECKED(UHairCardGeneratorPluginSettings, GenerateForGroomGroup))
+			|| (PropertyName == GET_MEMBER_NAME_CHECKED(UHairCardGeneratorPluginSettings, RandomSeed))
+			|| (PropertyName == GET_MEMBER_NAME_CHECKED(UHairCardGeneratorPluginSettings, AtlasSize))
+			|| (PropertyName == GET_MEMBER_NAME_CHECKED(UHairCardGeneratorPluginSettings, bUseReservedSpaceFromPreviousLOD))
+			|| (PropertyName == GET_MEMBER_NAME_CHECKED(UHairCardGeneratorPluginSettings, ReserveTextureSpaceLOD))
+			|| (PropertyName == GET_MEMBER_NAME_CHECKED(UHairCardGeneratorPluginSettings, bUseGroomAssetStrandWidth)));
+	}
+	
+	if ( (PropertyName == GET_MEMBER_NAME_CHECKED(UHairCardGeneratorPluginSettings, AtlasSize)) )
+	{
+		return !bUseReservedSpaceFromPreviousLOD;
+	}
+
+	if ( (PropertyName == GET_MEMBER_NAME_CHECKED(UHairCardGeneratorPluginSettings, bGenerateGeometryForAllGroups)) )
+	{
+		return (GenerateForGroomGroup == 0);
+	}
+
+	return true;
+}
+
+void UHairCardGeneratorPluginSettings::PostEditChangeProperty(FPropertyChangedEvent& PropertyChangedEvent)
+{
+	Super::PostEditChangeProperty(PropertyChangedEvent);
+
+	FName PropertyName = (PropertyChangedEvent.Property != nullptr) ? PropertyChangedEvent.Property->GetFName() : NAME_None;
+
+	if (PropertyName == GET_MEMBER_NAME_CHECKED(UHairCardGeneratorPluginSettings, BaseFilename))
+	{
+		if (BaseFilename.IsEmpty())
+		{
+			BaseFilename = HairCardSettings_Helpers::GetDefaultBaseNameForHairCards(GroomAsset, LODIndex);
+		}
+		else
+		{
+			// We use this name to generate paths and import object names so need it to be correctly sanitized (and replace "\", "/" characters)
+			BaseFilename = ObjectTools::SanitizeInvalidChars(BaseFilename, INVALID_OBJECTNAME_CHARACTERS TEXT("\\"));
+		}
+
+		if ( !ResetFromSettingsJson() )
+		{
+			UpdateOutputPaths();
+		}
+
+		// NOTE: This must be updated if basename or derived lod options change
+		UpdateParentInfo();
+	}
+	else if (PropertyName == GET_MEMBER_NAME_CHECKED(UHairCardGeneratorPluginSettings, DestinationPath))
+	{
+		if (DestinationPath.Path.IsEmpty())
+		{
+			DestinationPath.Path = HairCardSettings_Helpers::GetDefaultPathForHairCards(GroomAsset, LODIndex);
+		}
+	}
+	else if (PropertyName == GET_MEMBER_NAME_CHECKED(UHairCardGeneratorPluginSettings, bReduceCardsFromPreviousLOD) 
+		|| PropertyName == GET_MEMBER_NAME_CHECKED(UHairCardGeneratorPluginSettings, bUseReservedSpaceFromPreviousLOD))
+	{
+		UpdateParentInfo();
+	}
+}
+
+void UHairCardGeneratorPluginSettings::BuildDDCKey(FArchive& Ar)
+{
+	SerializeEditableSettings(Ar);
+
+	int FilterGroupSettingsCount = FilterGroupGenerationSettings.Num();
+	Ar << FilterGroupSettingsCount;
+
+	for ( int i=0; i < FilterGroupGenerationSettings.Num(); ++i )
+		FilterGroupGenerationSettings[i]->BuildDDCSubKey(Ar);
+}
+
+
+bool UHairCardGeneratorPluginSettings::IsCompatibleSettings(UHairCardGenerationSettings* SettingsObject)
+{
+	return SettingsObject && SettingsObject->IsA<UHairCardGeneratorPluginSettings>();
+}
+
+FString UHairCardGeneratorPluginSettings::GetGroomName() const
+{
+	return GroomAsset->GetName();
+}
+
+/* UHairCardGeneratorGroupSettings
+ *****************************************************************************/
+
+void UHairCardGeneratorGroupSettings::Reset(TObjectPtr<UHairCardGeneratorGroupSettings> TemplateObject)
+{
+	HairCardSettings_Helpers::ResetSettingsFromTemplate(this, TemplateObject);
+}
+
+void UHairCardGeneratorGroupSettings::Reset(const TSharedPtr<FJsonObject>& TemplateSettings)
+{
+	HairCardSettings_Converter::ResetSettingsFromJson(this, TemplateSettings);
+}
+
+void UHairCardGeneratorGroupSettings::ResetFilterDefault(const TArray<FName>& CardGroupNames)
+{
+	// HACK: Need to use a demonstrably invalid value by default to check against so we can only update from "default"
+	if ( !ApplyToCardGroups.Contains(NAME_None) )
+		return;
+
+	// Default filter includes all group ids
+	ApplyToCardGroups.Reset();
+	for ( FName id : CardGroupNames )
+		ApplyToCardGroups.Add(id);
+
+	// ApplyToCardGroups.Sort([](FName a, FName b){return a < b;});
+}
+
+void UHairCardGeneratorGroupSettings::UpdateStrandFilter()
+{
+	StrandsFilter = MakeUnique<FStrandFilterOp_InCardGroup>(ApplyToCardGroups);
+
+	StrandFilterChecksDirty = true;
+	SetStrandStats();
+}
+
+void UHairCardGeneratorGroupSettings::UpdateGenerateFilename(const FString& BaseFilename, int index)
+{
+	GenerateFilename = BaseFilename + TEXT("_GIDX") + FString::FromInt(index);
+}
+
+void UHairCardGeneratorGroupSettings::UpdateParentName(const FString& BaseParentName, int index)
+{
+	ParentName = TEXT("");
+	if ( !BaseParentName.IsEmpty() )
+	{
+		ParentName = BaseParentName + TEXT("_GIDX") + FString::FromInt(index);
+	}
+}
+
+void UHairCardGeneratorGroupSettings::RemoveSettingsFilterGroup()
+{
+	UHairCardGeneratorPluginSettings* PluginSettings = Cast<UHairCardGeneratorPluginSettings>(GetOuter());
+	PluginSettings->RemoveSettingsFilterGroup(this);
+}
+
+bool UHairCardGeneratorGroupSettings::CanEditChange(const FProperty* InProperty) const
+{
+	bool bCanEdit = Super::CanEditChange(InProperty);
+	if (!bCanEdit)
+	{
+		return false;
+	}
+
+	FName PropertyName = InProperty->GetFName();
+	if (PropertyName == GET_MEMBER_NAME_CHECKED(UHairCardGeneratorGroupSettings, MaxVerticalSegmentsPerCard))
+	{
+		return !UseAdaptiveSubdivision;
+	}
+	if (PropertyName == GET_MEMBER_NAME_CHECKED(UHairCardGeneratorGroupSettings, TargetTriangleCount))
+	{
+		return UseAdaptiveSubdivision;
+	}
+
+	return true;
+}
+
+void UHairCardGeneratorGroupSettings::PostEditChangeProperty(FPropertyChangedEvent& PropertyChangedEvent)
+{
+	Super::PostEditChangeProperty(PropertyChangedEvent);
+
+	FName PropertyName = (PropertyChangedEvent.Property != nullptr) ? PropertyChangedEvent.Property->GetFName() : NAME_None;
+	if (PropertyName == GET_MEMBER_NAME_CHECKED(UHairCardGeneratorGroupSettings, ApplyToCardGroups))
+	{
+		// ApplyToCardGroups.Sort([](FName a, FName b){return a < b;});
+		UpdateStrandFilter();
+	}
+}
+
+void UHairCardGeneratorGroupSettings::BuildDDCSubKey(FArchive& Ar)
+{
+	SerializeGroupSettings(Ar);
+}
+
+void UHairCardGeneratorGroupSettings::SetSource(TObjectPtr<const UGroomAsset> SourceGroom, int32 GenLODIndex)
+{
+	GroomAsset = SourceGroom;
+	LODIndex = GenLODIndex;
+}
+
+uint8 UHairCardGeneratorGroupSettings::GetFilterGroupGeneratedFlags() const
+{
+	// TODO: Cache these checks for UI and update on posteditchange
+	UHairCardGeneratorPluginSettings* PluginSettings = Cast<UHairCardGeneratorPluginSettings>(GetOuter());
+	uint8 PipelineFlags = PluginSettings->GetPipelineGeneratedDifferences();
+	if (PipelineFlags == (uint8)EHairCardGenerationSettingsCategories::All)
+		return PipelineFlags;
+
+	return PipelineFlags | GetPipelineFilterGroupDifferences();
+}
+
+bool UHairCardGeneratorGroupSettings::GetReduceFromPreviousLOD() const
+{
+	UHairCardGeneratorPluginSettings* PluginSettings = Cast<UHairCardGeneratorPluginSettings>(GetOuter());
+	return PluginSettings->bReduceCardsFromPreviousLOD;
+}
+
+FString UHairCardGeneratorGroupSettings::GetGeneratedFilterGroupSettingsFilename() const
+{
+	return UHairCardGeneratorPluginSettings::GetGeneratedSettingsPath() / GenerateFilename + TEXT(".settings.bin");
+}
+
+void UHairCardGeneratorGroupSettings::SerializeGroupSettings(FArchive& Ar)
+{
+	// NOTE: Only append new settings fields at end (FArchive produces reasonable defaults without buffer overflow)
+	Ar << ApplyToCardGroups;
+	Ar << TargetNumberOfCards;
+	Ar << MaxNumberOfFlyaways;
+	Ar << UseMultiCardClumps;
+
+	Ar << UseAdaptiveSubdivision;
+	Ar << MaxVerticalSegmentsPerCard;
+	Ar << TargetTriangleCount;
+
+	Ar << NumberOfTexturesInAtlas;
+	Ar << StrandWidthScalingFactor;
+	Ar << UseOptimizedCompressionFactor;
+}
+
+uint8 UHairCardGeneratorGroupSettings::GetPipelineFilterGroupDifferences() const
+{
+	const FString GroupSettingsBinFile = GetGeneratedFilterGroupSettingsFilename();
+	if ( !FPaths::FileExists(GroupSettingsBinFile) )
+		return (uint8)EHairCardGenerationPipeline::All;
+
+	TArray<uint8> SettingsBytes;
+	if ( !FFileHelper::LoadFileToArray(SettingsBytes, *GroupSettingsBinFile, 0) )
+		return (uint8)EHairCardGenerationPipeline::All;
+
+	FMemoryReader MemReader(SettingsBytes, true);
+
+	TObjectPtr<UHairCardGeneratorGroupSettings> DiffObject = NewObject<UHairCardGeneratorGroupSettings>();
+	DiffObject->SerializeGroupSettings(MemReader);
+
+	return HairCardSettings_Helpers::PipelineSettingsDiff(this, DiffObject);
+}
+
+void UHairCardGeneratorGroupSettings::ClearPipelineGeneratedFilterGroup() const
+{
+	const FString GroupSettingsBinFile = GetGeneratedFilterGroupSettingsFilename();
+
+	IPlatformFile& FileManager = FPlatformFileManager::Get().GetPlatformFile();
+	FileManager.DeleteFile(*GroupSettingsBinFile);
+}
+
+void UHairCardGeneratorGroupSettings::WritePipelineGeneratedFilterGroup()
+{
+	const FString GeneratedPath = UHairCardGeneratorPluginSettings::GetGeneratedSettingsPath();
+	IPlatformFile& FileManager = FPlatformFileManager::Get().GetPlatformFile();
+	if ( !FileManager.DirectoryExists(*GeneratedPath) )
+		FileManager.CreateDirectoryTree(*GeneratedPath);
+
+	TArray<uint8> SettingsBytes;
+	SettingsBytes.Reserve(64);
+	FMemoryWriter MemWriter(SettingsBytes, true);
+
+	SerializeGroupSettings(MemWriter);
+
+	const FString GroupSettingsBinFile = GetGeneratedFilterGroupSettingsFilename();
+	FFileHelper::SaveArrayToFile(SettingsBytes, *GroupSettingsBinFile);
+}
+
+void UHairCardGeneratorGroupSettings::SetStrandStats()
+{
+	if ( !GroomAsset )
+		return;
+
+	float Length = 0.0;
+	float AvgCurvRadius = 0.0;
+	int CountFiltered = 0;
+
+	FHairDescription HairDesc = GroomAsset->GetHairDescription();
+	for ( const FHairStrandAttribsRefProxy& StrandAttribs : FStrandAttributesRange(HairDesc, FHairStrandAttribsConstIterator::DefaultIgnore) )
+	{
+		if ( StrandsFilter->FilterStrand(StrandAttribs) )
+		{
+			CountFiltered ++;
+
+			MatrixXf VertexPositions(StrandAttribs.NumVerts(), 3);
+
+			for (int StrandVertIndex = 0; StrandVertIndex < StrandAttribs.NumVerts(); ++StrandVertIndex)
+			{
+				const FVector3f& VertexPosition = StrandAttribs.GetVertexPosition(StrandVertIndex);
+
+				for (int i=0; i<3; i++)
+				{
+					VertexPositions(StrandVertIndex, i) = VertexPosition[i];
+				}
+			}
+
+			Length += FHairCardGenCardSubdivider::GetCurveLength(VertexPositions);
+			AvgCurvRadius += FHairCardGenCardSubdivider::GetAverageCurvatureRadius(VertexPositions);
+		}
+	}
+
+	StrandCount = CountFiltered;
+}
\ No newline at end of file
diff --git a/Engine/Plugins/Experimental/HairCardGenerator/Source/HairCardGeneratorEditor/Private/HairCardGeneratorPluginSettings.h b/Engine/Plugins/Experimental/HairCardGenerator/Source/HairCardGeneratorEditor/Private/HairCardGeneratorPluginSettings.h
new file mode 100644
--- /dev/null
+++ b/Engine/Plugins/Experimental/HairCardGenerator/Source/HairCardGeneratorEditor/Private/HairCardGeneratorPluginSettings.h
@@ -0,0 +1,401 @@
+// Copyright Epic Games, Inc. All Rights Reserved.
+
+#pragma once
+
+#include "CoreMinimal.h"
+#include "GroomAssetCards.h" // for UHairCardGenerationSettings
+#include "HairCardGenStrandFilter.h"
+
+#include "HairCardGeneratorPluginSettings.generated.h"
+
+class FJsonObject;
+
+UENUM(BlueprintType)
+enum class EHairCardAtlasSize : uint8
+{
+	Uninitialized = 0  UMETA(Hidden),
+	AtlasSize1024 = 10 UMETA(DisplayName = "1024 x 1024"),
+	AtlasSize2048 = 11 UMETA(DisplayName = "2048 x 2048"),
+	AtlasSize4096 = 12 UMETA(DisplayName = "4096 x 4096"),
+	AtlasSize8192 = 13 UMETA(DisplayName = "8192 x 8192")
+};
+
+// Helper enumeration for settings dirty bits for pipeline steps
+UENUM()
+enum class EHairCardGenerationPipeline : uint8
+{
+	None = 0x0,
+	// Pipeline Steps
+	StrandClustering = 0x01,   // 1. Cluster strands into clumps for creating cards
+	GeometryGeneration = 0x02, // 2. Geometry generation from strand clusters (clumps)
+	TextureClustering = 0x04,  // 3. Cluster clumps into num_texures, and find the clump that's closest to cluster mean
+	TextureLayout = 0x08,      // 4. Layout representative clumps in texture atlas
+	TextureRendering = 0x10,   // 5. Render representative clumps to texture atlas
+	GenerateMesh = 0x20,       // 6. Generate a mesh from all clumps (card) geometry
+	ImportTextures = 0x40,     // 7. Import textures and assign to card asset slots
+
+	// HACK: These flag combos should probably be in a separate enum
+	// Useful shorthand for flag update checks
+	// Some update required
+	All = StrandClustering
+			| GeometryGeneration
+			| TextureClustering
+			| TextureLayout
+			| TextureRendering
+			| GenerateMesh
+			| ImportTextures,
+	// Group pipeline stage update required
+	GroupUpdate = StrandClustering
+				| GeometryGeneration
+				| TextureClustering,
+	// Textures will be updated
+	TextureUpdate = TextureLayout
+				| TextureRendering,
+	// Asset import/mesh regen required
+	ImportUpdate = GenerateMesh
+				| ImportTextures
+};
+
+ENUM_CLASS_FLAGS(EHairCardGenerationPipeline)
+
+UENUM()
+enum class EHairCardGenerationSettingsCategories : uint8
+{
+	// Settings Categories
+	None = 0x0,
+	All = (uint8)EHairCardGenerationPipeline::StrandClustering
+				| (uint8)EHairCardGenerationPipeline::GeometryGeneration
+				| (uint8)EHairCardGenerationPipeline::TextureClustering
+				| (uint8)EHairCardGenerationPipeline::TextureLayout
+				| (uint8)EHairCardGenerationPipeline::TextureRendering
+				| (uint8)EHairCardGenerationPipeline::GenerateMesh
+				| (uint8)EHairCardGenerationPipeline::ImportTextures,
+	Asset = All,
+	LevelOfDetail = All,
+	Randomness = All,
+	Cards = All,
+	Geometry = (uint8)EHairCardGenerationPipeline::GeometryGeneration
+				| (uint8)EHairCardGenerationPipeline::TextureClustering
+				| (uint8)EHairCardGenerationPipeline::TextureLayout
+				| (uint8)EHairCardGenerationPipeline::TextureRendering
+				| (uint8)EHairCardGenerationPipeline::GenerateMesh
+				| (uint8)EHairCardGenerationPipeline::ImportTextures,
+	Textures = (uint8)EHairCardGenerationPipeline::TextureClustering
+				| (uint8)EHairCardGenerationPipeline::TextureLayout
+				| (uint8)EHairCardGenerationPipeline::TextureRendering
+				| (uint8)EHairCardGenerationPipeline::GenerateMesh
+				| (uint8)EHairCardGenerationPipeline::ImportTextures,
+	TextureRendering = (uint8)EHairCardGenerationPipeline::TextureLayout
+				| (uint8)EHairCardGenerationPipeline::TextureRendering
+				| (uint8)EHairCardGenerationPipeline::GenerateMesh
+				| (uint8)EHairCardGenerationPipeline::ImportTextures,
+	Import = (uint8)EHairCardGenerationPipeline::GenerateMesh
+				| (uint8)EHairCardGenerationPipeline::ImportTextures,
+};
+
+
+/**
+ * Self contained hair card generation settings. Separated/Duplicated from the 
+ * groom asset to: (1) constrain the prompt window's details view, and (2) let
+ * the user modify it without mutating the original groom (settings should only 
+ * be applied/saved when the process is done).
+ * 
+ * UHairCardGenerationPluginSettings holds all the generation settings required
+ * to regenerate hair cards for a particular groom asset. Implements the
+ * UHairCardGenerationSettings interface so that it can be used with the
+ * hair card generation extension module (see IHairCardGenerator API).
+*/
+UCLASS(BlueprintType, config=Editor)
+class UHairCardGeneratorPluginSettings : public UHairCardGenerationSettings
+{
+	GENERATED_BODY()
+
+public:
+	// Destination content path for generated mesh/texture assets
+	UPROPERTY(EditAnywhere, Category="Import", meta=(ContentDir))
+	FDirectoryPath DestinationPath;
+
+	// Base filename of generated mesh/texture assets
+	UPROPERTY(EditAnywhere, BlueprintReadOnly, Category="Asset")
+	FString BaseFilename;
+
+	// Use previous LOD generated cards and textures but reduce triangle count and flyaways
+	UPROPERTY(EditAnywhere, BlueprintReadOnly, Category="Level Of Detail")
+	bool bReduceCardsFromPreviousLOD = false;
+
+	// Generate geometry for all groom groups on group 0
+	UPROPERTY(Transient, EditAnywhere, BlueprintReadonly, Category="Asset", meta=(AllowPrivateAccess))
+	bool bGenerateGeometryForAllGroups = true;
+
+	// Seed value for pseudo-random number generation (set to a specific value for repeatable results)
+	UPROPERTY(EditAnywhere, BlueprintReadOnly, Category="Randomness")
+	int RandomSeed = 0;
+
+	// Place new card textures in reserved space from previous LOD
+	UPROPERTY(EditAnywhere, BlueprintReadOnly, Category="Texture Rendering")
+	bool bUseReservedSpaceFromPreviousLOD = false;
+
+	// Size of hair card texture atlases
+	UPROPERTY(EditAnywhere, BlueprintReadOnly, Category="Texture Rendering")
+	EHairCardAtlasSize AtlasSize = EHairCardAtlasSize::AtlasSize4096;
+
+	// Percentage of texture atlas space to reserve for higher LODs 
+	UPROPERTY(EditAnywhere, BlueprintReadOnly, Category="Texture Rendering", meta=(ClampMin="0", ClampMax="75"))
+	int ReserveTextureSpaceLOD = 0;
+
+	// Use strand width
+	UPROPERTY(EditAnywhere, BlueprintReadOnly, Category="Texture Rendering")
+	bool bUseGroomAssetStrandWidth = true;
+
+	// Minimum strand depth value (mapped to 0 in depth texture)
+	UPROPERTY(BlueprintReadOnly, Category="Texture Rendering|Advanced", meta=(ClampMin="-10.0", ClampMax="0.0"))
+	float DepthMinimum = -2.0f;
+
+	// Maximum strand depth value (mapped to 1 in depth texture)
+	UPROPERTY(BlueprintReadOnly, Category="Texture Rendering|Advanced", meta=(ClampMin="0.0", ClampMax="10.0"))
+	float DepthMaximum = 2.0f;
+
+	// Number of iterations for smoothing strand UVs
+	UPROPERTY(BlueprintReadOnly, Category="Texture Rendering|Advanced", meta=(ClampMin="0", ClampMax="50"))
+	int32 UVSmoothingIterations = 25;
+	
+	UPROPERTY(Transient, BlueprintReadOnly, Category="Asset")
+	FString OutputPath;
+
+	void SetSource(TObjectPtr<const UGroomAsset> SourceObject, int32 GenLODIndex = 0, int32 PhysGroupIndex = 0);
+	// Load last-run values for options from card source description or defaults if no last-run info
+	void ResetFromSourceDescription(const FHairGroupsCardsSourceDescription& SourceDesc);
+	// Load settings from Json file if one exists
+	bool ResetFromSettingsJson();
+	// Load default values for all options
+	void ResetToDefault();
+	void ResetFilterGroupSettingsToDefault(int FilterGroupIndex);
+
+	// UI Hooks for adding/removing settings groups
+	void AddNewSettingsFilterGroup();
+	void RemoveSettingsFilterGroup(TObjectPtr<UHairCardGeneratorGroupSettings> SettingsGroup);
+
+	// Output settings to json file
+	void WriteGenerationSettings();
+
+	// Handle mapping settings differences to regeneration requirements
+	static bool CheckGenerationFlags(uint8 NeedsGenFlags, const EHairCardGenerationPipeline PipelineStage);
+	uint8 GetAllPipelineGeneratedDifferences() const;
+	uint8 GetPipelineGeneratedDifferences() const;
+	void ClearPipelineGenerated() const;
+	void WritePipelineGenerated();
+	uint8 GetPipelineFilterGroupDifferences(int FilterGroupIndex) const;
+	void ClearPipelineGeneratedFilterGroup(int FilterGroupIndex) const;
+	void WritePipelineGeneratedFilterGroup(int FilterGroupIndex);
+
+	// Handle checking and applying group settings strand filters
+	bool ValidateStrandFilters(TMap<uint32,uint32>& ErrorCountMap);
+	void UpdateStrandFilterAssignment();
+
+	//~ Begin UObject interface
+	virtual bool CanEditChange(const FProperty* InProperty) const override;
+	virtual void PostEditChangeProperty(FPropertyChangedEvent& PropertyChangedEvent) override;
+	//~ End UObject interface
+
+	//~ Begin UHairCardGenerationSettings interface
+	virtual void BuildDDCKey(FArchive& Ar) override;
+	//~ End UHairCardGenerationSettings interface
+
+	static bool IsCompatibleSettings(UHairCardGenerationSettings* SettingsObject);
+
+	static FString GetIntermediatePath();
+	static FString GetMetadataPath();
+	static FString GetGeneratedSettingsPath();
+
+	FString GetTextureImportBaseName() const;
+	FString GetTextureImportPath() const;
+	FString GetTextureContentPath() const;
+
+	bool HasValidFullParent() const;
+	bool HasDerivedTextureSettings() const;
+	int GetDerivedReservedTextureSize() const;
+	
+	FString GetGroomName() const;
+
+	int32 GetLODIndex() const { return LODIndex; }
+	int32 GetGenerateForGroomGroup() const { return GenerateForGroomGroup; }
+	void SetGenerateForGroomGroup(int GroupIndex) { GenerateForGroomGroup = GroupIndex; }
+
+	TArray<TObjectPtr<UHairCardGeneratorGroupSettings>>& GetFilterGroupSettings() { return FilterGroupGenerationSettings; }
+	const TArray<TObjectPtr<UHairCardGeneratorGroupSettings>>& GetFilterGroupSettings() const { return FilterGroupGenerationSettings; }
+
+	TObjectPtr<UHairCardGeneratorGroupSettings>& GetFilterGroupSettings(int Index) { return FilterGroupGenerationSettings[Index]; }
+	const TObjectPtr<UHairCardGeneratorGroupSettings>& GetFilterGroupSettings(int Index) const { return FilterGroupGenerationSettings[Index]; }
+
+private:
+	void UpdateOutputPaths();
+	void UpdateParentInfo();
+	bool FindDerivedTextureSettings();
+	TSharedPtr<FJsonObject> GetParentTextureSettings() const;
+
+	static TSharedPtr<FJsonObject> TraverseSettingsParents(const FString& SettingsName, TFunction<bool(const TSharedPtr<FJsonObject>&)> StopCond);
+
+	static FString GetSettingsFileFromBase(const FString& SettingsName);
+	FString GetSettingsFilename() const;
+	FString GetGeneratedSettingsFilename() const;
+	FString GetGeneratedGroupSettingsFilename(int GroupIndex) const;
+	void ResetFromJson();
+	void WriteToJson();
+
+	void SerializeEditableSettings(FArchive& Ar);
+
+	void ResetNumFilterGroups(int Count);
+	void ResetFilterGroupSettings(int FilterGroupIndex, TObjectPtr<UHairCardGeneratorGroupSettings> OptionsTemplate);
+	void ResetFilterGroupSettings(int FilterGroupIndex, const TSharedPtr<FJsonObject>& OptionsTemplate);
+
+	void ValidateStrandFiltersInternal();
+	int32 AssignStrandToFilterGroup(const FHairStrandAttribsRefProxy& StrandAttribs);
+	uint32 CheckAllStrandAssignments(const FHairStrandAttribsRefProxy& StrandAttribs);
+
+	UPROPERTY(VisibleAnywhere, BlueprintReadOnly, Category="Asset", meta=(AllowPrivateAccess))
+	int32 LODIndex = INDEX_NONE;
+
+	// Static HCS version number to force regeneration on tool updates
+	UPROPERTY(BlueprintReadOnly, Category="Asset", meta=(AllowPrivateAccess))
+	FString Version = TEXT("0.3");
+
+	UPROPERTY(BlueprintReadOnly, Category="Settings Groups", meta=(AllowPrivateAccess))
+	TArray<TObjectPtr<UHairCardGeneratorGroupSettings>> FilterGroupGenerationSettings;
+
+	// The groom group for which geometry will be generated
+	UPROPERTY(Transient, VisibleAnywhere, BlueprintReadonly, Category="Asset", meta=(AllowPrivateAccess))
+	int32 GenerateForGroomGroup = INDEX_NONE;
+
+	UPROPERTY(Transient, BlueprintReadOnly, Category="Asset", meta=(AllowPrivateAccess))
+	FString IntermediatePath;
+
+	// The base name of direct parent LOD level, used for resubdividing geometry or writing to reserved texture layout regions
+	UPROPERTY(Transient, BlueprintReadOnly, Category="Level Of Detail", meta=(AllowPrivateAccess))
+	FString BaseParentName;
+
+	// The base name of LOD from which texture info/outputs will be pulled (for writing to reserved texture layouts)
+	UPROPERTY(Transient, BlueprintReadOnly, Category="Level Of Detail", meta=(AllowPrivateAccess))
+	FString DerivedTextureSettingsName;
+
+	UPROPERTY(Transient, BlueprintReadOnly, Category="Settings Groups", meta=(AllowPrivateAccess))
+	TArray<FName> CardGroupIds;
+
+	UPROPERTY(Transient,BlueprintReadOnly, Category="Settings Groups", meta=(AllowPrivateAccess))
+	TArray<int> StrandFilterGroupIndexMap;
+
+	UPROPERTY(Transient)
+	TObjectPtr<const UGroomAsset> GroomAsset = nullptr;
+
+	UPROPERTY(Transient)
+	TObjectPtr<const UStaticMesh> OldGeneratedMesh = nullptr;
+
+	uint32 UnassignedStrandsCount;
+	TMap<uint32,uint32> StrandErrorCountMap;
+};
+
+
+/** 
+ * Actual generation settings. Held in TArray, with strand filters applied to determine
+ * which settings apply to each strand (See FComposableStrandFilter).
+ */
+UCLASS(BlueprintType, config=Editor)
+class UHairCardGeneratorGroupSettings : public UObject
+{
+	GENERATED_BODY()
+
+	friend class UHairCardGeneratorPluginSettings;
+private:
+	UPROPERTY(VisibleAnywhere, BlueprintReadOnly, Category="Settings Group", meta=(AllowPrivateAccess))
+	int32 StrandCount = 0;
+
+	UPROPERTY(VisibleAnywhere, BlueprintReadOnly, Category="Asset", meta=(AllowPrivateAccess))
+	FString GenerateFilename;
+
+	UPROPERTY(VisibleAnywhere, BlueprintReadOnly, Category="Asset", meta=(AllowPrivateAccess))
+	int32 LODIndex = INDEX_NONE;
+
+	UPROPERTY(Transient, BlueprintReadOnly, Category="Level Of Detail", meta=(AllowPrivateAccess))
+	FString ParentName;
+
+	UPROPERTY(Transient)
+	TObjectPtr<const UGroomAsset> GroomAsset = nullptr;
+
+	bool StrandFilterChecksDirty = false;
+	TUniquePtr<FComposableStrandFilterOp> StrandsFilter;
+public:
+	// Specify which strands (card group ids) this settings group applies to
+	UPROPERTY(EditAnywhere, BlueprintReadOnly, Category="Cards")
+	TSet<FName> ApplyToCardGroups = {NAME_None};
+
+	// Total number of cards to generate for this settings group
+	UPROPERTY(EditAnywhere, BlueprintReadOnly, Category="Cards", meta=(ClampMin="1", ClampMax="10000"))
+	int32 TargetNumberOfCards = 2000;
+
+	// Target number of triangles of the final mesh, only for adaptive subdivision
+	UPROPERTY(EditAnywhere, BlueprintReadOnly, Category="Geometry", meta=(ClampMin="2", ClampMax="1000000"))
+	int32 TargetTriangleCount = 20000;
+
+	// Maximum number of cards to assign to flyaway strands
+	UPROPERTY(EditAnywhere, BlueprintReadOnly, Category="Cards", meta=(ClampMin="0", ClampMax="1000"))
+	int MaxNumberOfFlyaways = 10;
+
+	// Generate multiple cards (3) per strand clump to give the appearance of volume
+	UPROPERTY(BlueprintReadOnly, Category="Geometry")
+	bool UseMultiCardClumps = true;
+
+	// Use adaptive subdivision when generating card geometry
+	UPROPERTY(EditAnywhere, BlueprintReadOnly, Category="Geometry|Advanced")
+	bool UseAdaptiveSubdivision = true;
+
+	// Maximum number of segments along each card (aligned with the strands), ignored for adaptive subdivision
+	UPROPERTY(EditAnywhere, BlueprintReadOnly, Category="Geometry|Advanced", meta=(ClampMin="1", ClampMax="15"))
+	int32 MaxVerticalSegmentsPerCard = 10;
+
+	// Number of strand clump textures per atlas
+	UPROPERTY(EditAnywhere, BlueprintReadOnly, Category="Textures", meta=(ClampMin="1", ClampMax="300"))
+	int32 NumberOfTexturesInAtlas = 75;
+
+	// Scaling factor mapping strand-width to pixel-size for coverage texture
+	UPROPERTY(EditAnywhere, BlueprintReadOnly, Category="Texture Rendering", meta=(ClampMin="0.0", ClampMax="300.0"))
+	float StrandWidthScalingFactor = 1.0f;
+
+	// Compress textures along strand direction to save vertical redolution,
+	// if strands are all moving in nearly the same direction
+	UPROPERTY(EditAnywhere, BlueprintReadOnly, Category="Texture Rendering")
+	bool UseOptimizedCompressionFactor = true;
+
+
+	void Reset(TObjectPtr<UHairCardGeneratorGroupSettings> TemplateObject);
+	void Reset(const TSharedPtr<FJsonObject>& TemplateSettings);
+	void ResetFilterDefault(const TArray<FName>& CardGroupIds);
+	void UpdateStrandFilter();
+	void UpdateGenerateFilename(const FString& BaseFilename, int index);
+	void UpdateParentName(const FString& BaseParentName, int index);
+
+	void RemoveSettingsFilterGroup();
+
+	//~ Begin UObject interface
+	virtual bool CanEditChange(const FProperty* InProperty) const override;
+	virtual void PostEditChangeProperty(FPropertyChangedEvent& PropertyChangedEvent) override;
+	//~ End UObject interface
+
+	uint8 GetFilterGroupGeneratedFlags() const;
+	bool GetReduceFromPreviousLOD() const;
+
+	void BuildDDCSubKey(FArchive& Ar);
+
+	int32 GetStrandCount() {return StrandCount;}
+
+private:
+	void SetSource(TObjectPtr<const UGroomAsset> SourceGroom, int32 GenLODIndex = 0);
+
+	FString GetGeneratedFilterGroupSettingsFilename() const;
+
+	void SerializeGroupSettings(FArchive& Ar);
+
+	uint8 GetPipelineFilterGroupDifferences() const;
+	void ClearPipelineGeneratedFilterGroup() const;
+	void WritePipelineGeneratedFilterGroup();
+
+	void SetStrandStats();
+};
diff --git a/Engine/Plugins/Experimental/HairCardGenerator/Source/HairCardGeneratorEditor/Public/IHairCardGeneratorEditor.h b/Engine/Plugins/Experimental/HairCardGenerator/Source/HairCardGeneratorEditor/Public/IHairCardGeneratorEditor.h
new file mode 100644
--- /dev/null
+++ b/Engine/Plugins/Experimental/HairCardGenerator/Source/HairCardGeneratorEditor/Public/IHairCardGeneratorEditor.h
@@ -0,0 +1,16 @@
+// Copyright Epic Games, Inc. All Rights Reserved.
+
+#pragma once
+
+#include "Modules/ModuleInterface.h"
+#include "Modules/ModuleManager.h"
+#include "IHairCardGenerator.h"
+
+class IHairCardGeneratorEditor : public IModuleInterface, public IHairCardGenerator
+{
+public:
+	static FORCEINLINE IHairCardGeneratorEditor& Get()
+	{
+		return FModuleManager::LoadModuleChecked<IHairCardGeneratorEditor>("HairCardGeneratorEditor");
+	}
+};
