HairCardSolver: Fix several LOCTEXT key collisions.

diff --git a/Engine/Plugins/Experimental/HairCardGenerator/Source/HairCardGeneratorEditor/Private/HairCardGenOptionsWindow.cpp b/Engine/Plugins/Experimental/HairCardGenerator/Source/HairCardGeneratorEditor/Private/HairCardGenOptionsWindow.cpp
--- a/Engine/Plugins/Experimental/HairCardGenerator/Source/HairCardGeneratorEditor/Private/HairCardGenOptionsWindow.cpp
+++ b/Engine/Plugins/Experimental/HairCardGenerator/Source/HairCardGeneratorEditor/Private/HairCardGenOptionsWindow.cpp
@@ -312,7 +312,7 @@ TSharedRef<SWidget> SGroupSettingsListRow::GenerateWidgetForColumn(const FName&
 			.Value(this, &SGroupSettingsListRow::ReadTargetNumTextures)
 			.OnValueChanged(this, &SGroupSettingsListRow::WriteTargetNumTextures)
 			.IsEnabled(this, &SGroupSettingsListRow::CheckDisabledOnReducePrevLOD)
-			.ToolTipText(LOCTEXT("HairCardSettings.CardGroupFilter.ToolTip", "The number of textures to generate (must be <= number of cards)"));
+			.ToolTipText(LOCTEXT("HairCardSettings.TargetNumTextures.ToolTip.InvalidTextureNum", "The number of textures to generate (must be <= number of cards)"));
 	}
 	else if ( ColumnName == TEXT("MaxFlyaways") )
 	{
@@ -331,19 +331,19 @@ TSharedRef<SWidget> SGroupSettingsListRow::GenerateWidgetForColumn(const FName&
 			.Value(this, &SGroupSettingsListRow::ReadTargetNumTriangles)
 			.OnValueChanged(this, &SGroupSettingsListRow::WriteTargetNumTriangles)
 			.IsEnabled(this, &SGroupSettingsListRow::CheckEditableTriangles)
-			.ToolTipText(LOCTEXT("HairCardSettings.CardGroupFilter.ToolTip", "The total number of triangles to use for all cards (approximate)"));
+			.ToolTipText(LOCTEXT("HairCardSettings.TargetTriangles.ToolTip", "The total number of triangles to use for all cards (approximate)"));
 	}
 	else if ( ColumnName == TEXT("NumStrands") )
 	{
 		return SNew(STextBlock)
 			.Text(this, &SGroupSettingsListRow::GetStrandCountText)
-			.ToolTipText(LOCTEXT("HairCardSettings.CardGroupFilter.ToolTip", "The number of strands these settings will apply to"));
+			.ToolTipText(LOCTEXT("HairCardSettings.NumStrands.ToolTip", "The number of strands these settings will apply to"));
 	}
 	else if ( ColumnName == TEXT("IsGenerated") )
 	{
 		return SNew(STextBlock)
 			.Text(this, &SGroupSettingsListRow::GetGeneratedStateText)
-			.ToolTipText(LOCTEXT("HairCardSettings.CardGroupFilter.ToolTip", "Indicates whether cards and textures have already been generated using these settings"));
+			.ToolTipText(LOCTEXT("HairCardSettings.IsGenerated.ToolTip", "Indicates whether cards and textures have already been generated using these settings"));
 	}
 	else if ( ColumnName == TEXT("AddDeleteRow") )
 	{
diff --git a/Engine/Plugins/Experimental/HairCardGenerator/Source/HairCardGeneratorEditor/Private/HairCardGeneratorEditorModule.cpp b/Engine/Plugins/Experimental/HairCardGenerator/Source/HairCardGeneratorEditor/Private/HairCardGeneratorEditorModule.cpp
--- a/Engine/Plugins/Experimental/HairCardGenerator/Source/HairCardGeneratorEditor/Private/HairCardGeneratorEditorModule.cpp
+++ b/Engine/Plugins/Experimental/HairCardGenerator/Source/HairCardGeneratorEditor/Private/HairCardGeneratorEditorModule.cpp
@@ -219,7 +219,7 @@ static bool HairCardGeneratorEditor_Impl::LoadGroomData(UGroomAsset* NewGroomAss
 		if ( !HairCardGeneratorEditor_Impl::BuildHairCardGenGroomData(CardGenController->GetPointsPerCurve(), NewGroomAsset->GetHairDescription(), RestructuredData))
 			return false;
 
-		FScopedSlowTask SlowTask(0, LOCTEXT("GeneratingHairCards", "Loading Groom..."));
+		FScopedSlowTask SlowTask(0, LOCTEXT("GeneratingHairCards.LoadGroom", "Loading Groom..."));
 		SlowTask.MakeDialog(/*bShowCancelButton =*/true);
 
 		FString CachedGroomsPath = FPaths::ProjectIntermediateDir() / TEXT("GroomHairCardGen") / TEXT("CachedGrooms");
@@ -418,7 +418,7 @@ static bool HairCardGeneratorEditor_Impl::GenerateCardsForCardGroup(TObjectPtr<c
 		return false;
 	}
 
-	FScopedSlowTask SlowTask(10, FText::Format(LOCTEXT("GeneratingHairCards", "Generating Hair Cards for settings group {0}..."), index));
+	FScopedSlowTask SlowTask(10, FText::Format(LOCTEXT("GeneratingHairCards.GenerateGroup", "Generating Hair Cards for settings group {0}..."), index));
 	SlowTask.MakeDialog(/*bShowCancelButton =*/true);
 
 	SlowTask.EnterProgressFrame(1);
@@ -818,7 +818,7 @@ bool FHairCardGeneratorEditorModule::GenerateHairCardsForLOD(UGroomAsset* NewGro
 	// Generate static mesh from cards and add to destination package
 	if ( GenerationSettings->CheckGenerationFlags(AllGenFlags, EHairCardGenerationPipeline::GenerateMesh) )
 	{
-		FScopedSlowTask SlowTask(0, LOCTEXT("GeneratingHairCards", "Building mesh..."));
+		FScopedSlowTask SlowTask(0, LOCTEXT("GeneratingHairCards.BuildMesh", "Building mesh..."));
 		SlowTask.MakeDialog(/*bShowCancelButton =*/true);
 
 		// Don't bother adding _PG# at end of mesh name if there's only one physics group
