[HairCardSolver] Relax interpolation parameter to avoid artefacts when grooms misbehave

diff --git a/Engine/Plugins/Experimental/HairCardGenerator/Source/HairCardGeneratorEditor/Private/HairCardGenStrandNCSInterpolator.cpp b/Engine/Plugins/Experimental/HairCardGenerator/Source/HairCardGeneratorEditor/Private/HairCardGenStrandNCSInterpolator.cpp
--- a/Engine/Plugins/Experimental/HairCardGenerator/Source/HairCardGeneratorEditor/Private/HairCardGenStrandNCSInterpolator.cpp
+++ b/Engine/Plugins/Experimental/HairCardGenerator/Source/HairCardGeneratorEditor/Private/HairCardGenStrandNCSInterpolator.cpp
@@ -27,7 +27,7 @@ FHairCardGenStrandNSCInterpolator::FHairCardGenStrandNSCInterpolator(const TArra
 	}
 	for (int i = 1; i < NumPoints; ++i)
 	{
-		DistancesFromStart[i] /= DistancesFromStart[NumPoints - 1];
+		DistancesFromStart[i] = (DistancesFromStart[i] + float(i) * 0.1 * DistancesFromStart.Last()) / ((0.1 * float(NumPoints) + 0.9) * DistancesFromStart.Last());
 	}
 
 	Splines.Reserve(4);
