[HairCardSolver] Fix bug when reducing from a previous LOD with higher number of flyaways or a small number of triangles

diff --git a/Engine/Plugins/Experimental/HairCardGenerator/Source/HairCardGeneratorEditor/Private/HairCardGenCardSubdivider.cpp b/Engine/Plugins/Experimental/HairCardGenerator/Source/HairCardGeneratorEditor/Private/HairCardGenCardSubdivider.cpp
--- a/Engine/Plugins/Experimental/HairCardGenerator/Source/HairCardGeneratorEditor/Private/HairCardGenCardSubdivider.cpp
+++ b/Engine/Plugins/Experimental/HairCardGenerator/Source/HairCardGeneratorEditor/Private/HairCardGenCardSubdivider.cpp
@@ -186,7 +186,7 @@ TArray<float> FHairCardGenCardSubdivider::GetSubdivisionPoints(const MatrixXf Po
 
 float FHairCardGenCardSubdivider::GetToleranceFromLengthAndCurvRadius(float Length, float CurvRadius, float Subdivisions)
 {
-	if (Subdivisions < 1.01) return 1e10;
+	if (Subdivisions < 1.) Subdivisions = 1.;
 
 	float Num = Length * Length * Length;
 	float Den = 3 * CurvRadius * CurvRadius * (2 * Subdivisions - 1) * (2 * Subdivisions - 1) * (2 * Subdivisions - 1);
