[HairCardSolver] Refine triangle count accuracy when it's small enough that some cards are not subdivided

diff --git a/Engine/Plugins/Experimental/HairCardGenerator/Source/HairCardGeneratorEditor/Private/HairCardGenCardSubdivider.cpp b/Engine/Plugins/Experimental/HairCardGenerator/Source/HairCardGeneratorEditor/Private/HairCardGenCardSubdivider.cpp
--- a/Engine/Plugins/Experimental/HairCardGenerator/Source/HairCardGeneratorEditor/Private/HairCardGenCardSubdivider.cpp
+++ b/Engine/Plugins/Experimental/HairCardGenerator/Source/HairCardGeneratorEditor/Private/HairCardGenCardSubdivider.cpp
@@ -172,9 +172,17 @@ TArray<float> FHairCardGenCardSubdivider::GetSubdivisionPoints(const MatrixXf Po
 
 	if (UseAdaptive)
 	{
-		SubdTolerance = SubdToleranceInitial;
-		MeanCurve = LinearInterpolation(Points, NumInterpPoints);
-		SubdPoints = GetAdaptiveSubdPoints();
+		if (SubdToleranceInitial > 0.)
+		{
+			SubdTolerance = SubdToleranceInitial;
+
+			MeanCurve = LinearInterpolation(Points, NumInterpPoints);
+			SubdPoints = GetAdaptiveSubdPoints();
+		}
+		else
+		{
+			SubdPoints = MatrixXfToTArray(LinearInterpolation(Points, 2));
+		}
 	}
 	else
 	{
@@ -186,10 +194,10 @@ TArray<float> FHairCardGenCardSubdivider::GetSubdivisionPoints(const MatrixXf Po
 
 float FHairCardGenCardSubdivider::GetToleranceFromLengthAndCurvRadius(float Length, float CurvRadius, float Subdivisions)
 {
-	if (Subdivisions < 1.) Subdivisions = 1.;
+	if (Subdivisions <= 0.5) return -1.;
 
 	float Num = Length * Length * Length;
-	float Den = 3 * CurvRadius * CurvRadius * (2 * Subdivisions - 1) * (2 * Subdivisions - 1) * (2 * Subdivisions - 1);
+	float Den = 3. * CurvRadius * CurvRadius * (2. * Subdivisions - 1.) * (2. * Subdivisions - 1.) * (2. * Subdivisions - 1.);
 
 	return Num / Den;
 }
\ No newline at end of file
diff --git a/Engine/Plugins/Experimental/HairCardGenerator/Source/HairCardGeneratorEditor/Private/HairCardGeneratorEditorModule.cpp b/Engine/Plugins/Experimental/HairCardGenerator/Source/HairCardGeneratorEditor/Private/HairCardGeneratorEditorModule.cpp
--- a/Engine/Plugins/Experimental/HairCardGenerator/Source/HairCardGeneratorEditor/Private/HairCardGeneratorEditorModule.cpp
+++ b/Engine/Plugins/Experimental/HairCardGenerator/Source/HairCardGeneratorEditor/Private/HairCardGeneratorEditorModule.cpp
@@ -363,29 +363,40 @@ static bool HairCardGeneratorEditor_Impl::GenerateCardsGeometry(TObjectPtr<const
 		return FHairCardGenCardSubdivider::GetToleranceFromLengthAndCurvRadius(AvgLength, AvgCurvRadius, NumSubdivisions);
 	};
 
+	int const InitialTargetTriCount = Settings->GetFilterGroupSettings(index)->TargetTriangleCount;
+
 	FHairCardGenCardSubdivider Subdivider(
-		GetSubdToleranceFromTriCount(Settings->GetFilterGroupSettings(index)->TargetTriangleCount),
+		GetSubdToleranceFromTriCount(InitialTargetTriCount),
 		Settings->GetFilterGroupSettings(index)->UseAdaptiveSubdivision,
 		Settings->GetFilterGroupSettings(index)->MaxVerticalSegmentsPerCard);
 
-	auto CountTriangles = [&Subdivider, &AvgCurves]() 
+	if (InitialTargetTriCount <= 2 * TotalNumCards)
 	{
-		int TriCount = 0;
-		for (int CardId = 0; CardId < AvgCurves.Num(); CardId++) TriCount += 2 * (Subdivider.GetSubdivisionPoints(AvgCurves[CardId]).Num() / 3 - 1);
-		return TriCount;
-	};
+		Subdivider.SetSubdTolerance(-1.);
+	}
+	else
+	{
+		auto CountTriangles = [&Subdivider, &AvgCurves]() 
+		{
+			int TriCount = 0;
+			for (int CardId = 0; CardId < AvgCurves.Num(); CardId++) TriCount += 2 * (Subdivider.GetSubdivisionPoints(AvgCurves[CardId]).Num() / 3 - 1);
+			return TriCount;
+		};
 
-	const int InitialTargetTriCount = Settings->GetFilterGroupSettings(index)->TargetTriangleCount;
-	const int InitialTriCount = CountTriangles();
+		const int InitialTriCount = CountTriangles();
 
-	const int ShiftedTargetTriCount = 2 * InitialTargetTriCount - InitialTriCount;
-	Subdivider.SetSubdTolerance(GetSubdToleranceFromTriCount(ShiftedTargetTriCount));
-	const int ShiftedTriCount = CountTriangles();
+		int ShiftedTargetTriCount = 2 * InitialTargetTriCount - InitialTriCount;
+		int const MinTargetTriCount = 2.2 * float(TotalNumCards * InitialTargetTriCount) / float(InitialTriCount);
+		if (ShiftedTargetTriCount < MinTargetTriCount) ShiftedTargetTriCount = MinTargetTriCount;
 
-	if (InitialTriCount != ShiftedTriCount)
-	{
-		const int CorrectedTargetTriCount = InitialTargetTriCount + int(float(ShiftedTargetTriCount - InitialTargetTriCount) / float(ShiftedTriCount - InitialTriCount) * float(InitialTargetTriCount - InitialTriCount));
-		Subdivider.SetSubdTolerance(GetSubdToleranceFromTriCount(CorrectedTargetTriCount));
+		Subdivider.SetSubdTolerance(GetSubdToleranceFromTriCount(ShiftedTargetTriCount));
+		const int ShiftedTriCount = CountTriangles();
+
+		if (InitialTriCount != ShiftedTriCount)
+		{
+			const int CorrectedTargetTriCount = InitialTargetTriCount + int(float(ShiftedTargetTriCount - InitialTargetTriCount) / float(ShiftedTriCount - InitialTriCount) * float(InitialTargetTriCount - InitialTriCount));
+			Subdivider.SetSubdTolerance(GetSubdToleranceFromTriCount(CorrectedTargetTriCount));
+		}
 	}
 
 	CardId = 0;
