HairCardSolver: Reworked settings creation so there is a single PostResetUpdates which is always run on new cards creation init.

diff --git a/Engine/Plugins/Experimental/HairCardGenerator/Source/HairCardGeneratorEditor/Private/HairCardGenControllerBase.cpp b/Engine/Plugins/Experimental/HairCardGenerator/Source/HairCardGeneratorEditor/Private/HairCardGenControllerBase.cpp
--- a/Engine/Plugins/Experimental/HairCardGenerator/Source/HairCardGeneratorEditor/Private/HairCardGenControllerBase.cpp
+++ b/Engine/Plugins/Experimental/HairCardGenerator/Source/HairCardGeneratorEditor/Private/HairCardGenControllerBase.cpp
@@ -106,16 +106,12 @@ void UHairCardGenControllerBase::CreateCardsStaticMesh(UStaticMesh* StaticMesh,
 	StaticMesh->CommitMeshDescription(0);
 }
 
-TObjectPtr<UHairCardGeneratorPluginSettings>& UHairCardGenControllerBase::GetGroomSettings(TObjectPtr<UGroomAsset> Groom, int LODIndex, int GroupIndex)
+TObjectPtr<UHairCardGeneratorPluginSettings>& UHairCardGenControllerBase::GetGroomSettings(TObjectPtr<UGroomAsset> Groom, int LODIndex)
 {
 	const FString SettingsTag = Groom->GetPathName() + TEXT(":LOD") + FString::FromInt(LODIndex);
 
-	// HACK: Could keep multiple settings objects around per-lod and group, but the group id should use the same settings so just update group id
-	TObjectPtr<UHairCardGeneratorPluginSettings>& Settings = GroomSettingsMap.FindOrAdd(SettingsTag);
-	if ( Settings)
-		Settings->SetGenerateForGroomGroup(GroupIndex);
-
-	return Settings;
+	// NOTE: Only keep around one settings object per-lod
+	return GroomSettingsMap.FindOrAdd(SettingsTag);
 }
 
 void UHairCardGenControllerBase::UpdateGroomSettings(TObjectPtr<UGroomAsset> Groom, int LODIndex, int GroupID, TObjectPtr<UHairCardGeneratorPluginSettings> NewSettings)
diff --git a/Engine/Plugins/Experimental/HairCardGenerator/Source/HairCardGeneratorEditor/Private/HairCardGenControllerBase.h b/Engine/Plugins/Experimental/HairCardGenerator/Source/HairCardGeneratorEditor/Private/HairCardGenControllerBase.h
--- a/Engine/Plugins/Experimental/HairCardGenerator/Source/HairCardGeneratorEditor/Private/HairCardGenControllerBase.h
+++ b/Engine/Plugins/Experimental/HairCardGenerator/Source/HairCardGeneratorEditor/Private/HairCardGenControllerBase.h
@@ -114,7 +114,7 @@ public:
 	void CreateCardsStaticMesh(UStaticMesh* StaticMesh, const TArray<float>& verts, const TArray<int32>& faces, const TArray<float>& normals, const TArray<float>& uvs, const TArray<int32>& groups);
 
 
-	TObjectPtr<UHairCardGeneratorPluginSettings>& GetGroomSettings(TObjectPtr<UGroomAsset> Groom, int LODIndex, int GroupID);
+	TObjectPtr<UHairCardGeneratorPluginSettings>& GetGroomSettings(TObjectPtr<UGroomAsset> Groom, int LODIndex);
 	void UpdateGroomSettings(TObjectPtr<UGroomAsset> Groom, int LODIndex, int GroupID, TObjectPtr<UHairCardGeneratorPluginSettings> NewSettings);
 
 private:
diff --git a/Engine/Plugins/Experimental/HairCardGenerator/Source/HairCardGeneratorEditor/Private/HairCardGeneratorEditorModule.cpp b/Engine/Plugins/Experimental/HairCardGenerator/Source/HairCardGeneratorEditor/Private/HairCardGeneratorEditorModule.cpp
--- a/Engine/Plugins/Experimental/HairCardGenerator/Source/HairCardGeneratorEditor/Private/HairCardGeneratorEditorModule.cpp
+++ b/Engine/Plugins/Experimental/HairCardGenerator/Source/HairCardGeneratorEditor/Private/HairCardGeneratorEditorModule.cpp
@@ -94,7 +94,6 @@ namespace HairCardGeneratorEditor_Impl
 	static TObjectPtr<UGroomAsset> ParseGroomAssetArg(const TArray<FString>& Args, bool bOptional);
 	static FHairGroupsCardsSourceDescription* FindCardSourceDescription(TObjectPtr<UGroomAsset> Groom, int32 LODIndex);
 	static TObjectPtr<UHairCardGeneratorPluginSettings> CreateGroomSettingsLast(TObjectPtr<UGroomAsset> Groom, const FHairGroupsCardsSourceDescription& SourceDesc);
-	static TObjectPtr<UHairCardGeneratorPluginSettings> CreateGroomSettingsDefault(TObjectPtr<UGroomAsset> Groom, int32 LODIndex, int32 PhysGroupIndex);
 
 	static void DbgCmdCacheGroom(const TArray<FString>& Args);
 
@@ -181,15 +180,6 @@ static TObjectPtr<UHairCardGeneratorPluginSettings> HairCardGeneratorEditor_Impl
 	return NewSettings;
 }
 
-static TObjectPtr<UHairCardGeneratorPluginSettings> HairCardGeneratorEditor_Impl::CreateGroomSettingsDefault(TObjectPtr<UGroomAsset> Groom, int32 LODIndex, int32 PhysGroupIndex)
-{
-	TObjectPtr<UHairCardGeneratorPluginSettings> NewSettings = NewObject<UHairCardGeneratorPluginSettings>(/*Outer =*/Groom);
-	NewSettings->SetSource(Groom, LODIndex, PhysGroupIndex);
-	NewSettings->ResetToDefault();
-
-	return NewSettings;
-}
-
 static void HairCardGeneratorEditor_Impl::DbgCmdCacheGroom(const TArray<FString>& Args)
 {
 	TObjectPtr<UGroomAsset> GroomAsset = HairCardGeneratorEditor_Impl::ParseGroomAssetArg(Args, false);
@@ -668,9 +658,17 @@ bool FHairCardGeneratorEditorModule::GenerateHairCardsForLOD(UGroomAsset* NewGro
 
 	UE_CLOG(CardsDesc.GenerationSettings && !IsCompatibleSettings(CardsDesc.GenerationSettings), LogHairCardGenerator, Warning, TEXT("Old hair-card generation settings (from a different generator) will be discarded when you run this generator."));
 
-	TObjectPtr<UHairCardGeneratorPluginSettings>& GenerationSettings = CardGenController->GetGroomSettings(NewGroomAsset, CardsDesc.LODIndex, CardsDesc.GroupIndex);
+	TObjectPtr<UHairCardGeneratorPluginSettings>& GenerationSettings = CardGenController->GetGroomSettings(NewGroomAsset, CardsDesc.LODIndex);
 	if ( !GenerationSettings )
+	{
 		GenerationSettings = HairCardGeneratorEditor_Impl::CreateGroomSettingsLast(NewGroomAsset, CardsDesc);
+	}
+	else
+	{
+		// Need to update all non-generation UI settings in case they've changed on groom asset
+		GenerationSettings->SetGenerateForGroomGroup(CardsDesc.GroupIndex);
+		GenerationSettings->PostResetUpdates();
+	}
 
 	bool bSuccess = false;
 	if ( !HairCardGenWindow_Utils::PromptUserWithHairCardGenDialog(GenerationSettings) )
diff --git a/Engine/Plugins/Experimental/HairCardGenerator/Source/HairCardGeneratorEditor/Private/HairCardGeneratorPluginSettings.cpp b/Engine/Plugins/Experimental/HairCardGenerator/Source/HairCardGeneratorEditor/Private/HairCardGeneratorPluginSettings.cpp
--- a/Engine/Plugins/Experimental/HairCardGenerator/Source/HairCardGeneratorEditor/Private/HairCardGeneratorPluginSettings.cpp
+++ b/Engine/Plugins/Experimental/HairCardGenerator/Source/HairCardGeneratorEditor/Private/HairCardGeneratorPluginSettings.cpp
@@ -325,7 +325,9 @@ void UHairCardGeneratorPluginSettings::ResetFromSourceDescription(const FHairGro
 	OldGeneratedMesh = SourceDesc.ImportedMesh;
 
 	if ( LastSettings == nullptr )
+	{
 		LastSettings = UHairCardGeneratorPluginSettings::StaticClass()->GetDefaultObject<UHairCardGeneratorPluginSettings>();
+	}
 
 	HairCardSettings_Helpers::ResetSettingsFromTemplate(this, LastSettings);
 	if (OldGeneratedMesh && LODIndex == LastSettings->LODIndex)
@@ -346,7 +348,9 @@ void UHairCardGeneratorPluginSettings::ResetFromSourceDescription(const FHairGro
 	{
 		TObjectPtr<UHairCardGeneratorGroupSettings> oldOpts;
 		if ( LastSettings->FilterGroupGenerationSettings.Num() > i )
+		{
 			oldOpts = LastSettings->FilterGroupGenerationSettings[i];
+		}
 		else
 		{
 			oldOpts = UHairCardGeneratorGroupSettings::StaticClass()->GetDefaultObject<UHairCardGeneratorGroupSettings>();
@@ -356,8 +360,7 @@ void UHairCardGeneratorPluginSettings::ResetFromSourceDescription(const FHairGro
 	}
 	ResetNumFilterGroups(GroupSettingsCount);
 
-	UpdateOutputPaths();
-	UpdateParentInfo();
+	PostResetUpdates();
 }
 
 bool UHairCardGeneratorPluginSettings::ResetFromSettingsJson()
@@ -376,8 +379,7 @@ bool UHairCardGeneratorPluginSettings::ResetFromSettingsJson()
 		return false;
 
 	ResetFromJson();
-	UpdateOutputPaths();
-	UpdateParentInfo();
+	PostResetUpdates();
 
 	return true;
 }
@@ -396,8 +398,7 @@ void UHairCardGeneratorPluginSettings::ResetToDefault()
 	ResetFilterGroupSettings(0, UHairCardGeneratorGroupSettings::StaticClass()->GetDefaultObject<UHairCardGeneratorGroupSettings>());
 	ResetNumFilterGroups(1);
 
-	UpdateOutputPaths();
-	UpdateParentInfo();
+	PostResetUpdates();
 }
 
 void UHairCardGeneratorPluginSettings::ResetFilterGroupSettingsToDefault(int FilterGroupIndex)
@@ -526,6 +527,12 @@ void UHairCardGeneratorPluginSettings::WritePipelineGeneratedFilterGroup(int Fil
 	FilterGroupGenerationSettings[FilterGroupIndex]->WritePipelineGeneratedFilterGroup();
 }
 
+void UHairCardGeneratorPluginSettings::PostResetUpdates()
+{
+	UpdateOutputPaths();
+	UpdateParentInfo();
+}
+
 void UHairCardGeneratorPluginSettings::UpdateOutputPaths()
 {
 	IntermediatePath = GetIntermediatePath();
diff --git a/Engine/Plugins/Experimental/HairCardGenerator/Source/HairCardGeneratorEditor/Private/HairCardGeneratorPluginSettings.h b/Engine/Plugins/Experimental/HairCardGenerator/Source/HairCardGeneratorEditor/Private/HairCardGeneratorPluginSettings.h
--- a/Engine/Plugins/Experimental/HairCardGenerator/Source/HairCardGeneratorEditor/Private/HairCardGeneratorPluginSettings.h
+++ b/Engine/Plugins/Experimental/HairCardGenerator/Source/HairCardGeneratorEditor/Private/HairCardGeneratorPluginSettings.h
@@ -227,6 +227,8 @@ public:
 	TObjectPtr<UHairCardGeneratorGroupSettings>& GetFilterGroupSettings(int Index) { return FilterGroupGenerationSettings[Index]; }
 	const TObjectPtr<UHairCardGeneratorGroupSettings>& GetFilterGroupSettings(int Index) const { return FilterGroupGenerationSettings[Index]; }
 
+	void PostResetUpdates();
+
 private:
 	void UpdateOutputPaths();
 	void UpdateParentInfo();
