[HairCardSolver] Regenerate textures when hair width strand attributes change

diff --git a/Engine/Plugins/Experimental/HairCardGenerator/Source/HairCardGeneratorEditor/Private/HairCardGenControllerBase.h b/Engine/Plugins/Experimental/HairCardGenerator/Source/HairCardGeneratorEditor/Private/HairCardGenControllerBase.h
--- a/Engine/Plugins/Experimental/HairCardGenerator/Source/HairCardGeneratorEditor/Private/HairCardGenControllerBase.h
+++ b/Engine/Plugins/Experimental/HairCardGenerator/Source/HairCardGeneratorEditor/Private/HairCardGenControllerBase.h
@@ -73,9 +73,6 @@ public:
 	UFUNCTION(BlueprintImplementableEvent, Category = Python)
 	bool LoadSettings(const UHairCardGeneratorPluginSettings* GeneratorSettings);
 
-	UFUNCTION(BlueprintImplementableEvent, Category = Python)
-	bool UpdateWidthMultipliers(const TArray<float>& HairWidths, const TArray<float>& RootScales, const TArray<float>& TipScales);
-
 	UFUNCTION(BlueprintImplementableEvent, Category = Python)
 	bool GenerateCardsGeometry();
 
diff --git a/Engine/Plugins/Experimental/HairCardGenerator/Source/HairCardGeneratorEditor/Private/HairCardGeneratorEditorModule.cpp b/Engine/Plugins/Experimental/HairCardGenerator/Source/HairCardGeneratorEditor/Private/HairCardGeneratorEditorModule.cpp
--- a/Engine/Plugins/Experimental/HairCardGenerator/Source/HairCardGeneratorEditor/Private/HairCardGeneratorEditorModule.cpp
+++ b/Engine/Plugins/Experimental/HairCardGenerator/Source/HairCardGeneratorEditor/Private/HairCardGeneratorEditorModule.cpp
@@ -719,28 +719,6 @@ bool FHairCardGeneratorEditorModule::GenerateHairCardsForLOD(UGroomAsset* NewGro
 		return false;
 	}
 
-	TArray <float> HairWidths, RootScales, TipScales;
-	HairWidths.Init(-1.0, NumPhysGroups);
-	RootScales.Init(1.0, NumPhysGroups);
-	TipScales.Init(1.0, NumPhysGroups);
-	if ( GenerationSettings->bUseGroomAssetStrandWidth )
-	{
-		for ( int PhysGroupIdx = 0; PhysGroupIdx<NumPhysGroups; PhysGroupIdx++ )
-		{
-			const FHairGeometrySettings& GeometrySettings = NewGroomAsset->GetHairGroupsRendering()[PhysGroupIdx].GeometrySettings;
-			if ( GeometrySettings.HairWidth_Override )
-			{
-				HairWidths[PhysGroupIdx] = GeometrySettings.HairWidth;
-			}
-			RootScales[PhysGroupIdx] = GeometrySettings.HairRootScale;
-			TipScales[PhysGroupIdx] = GeometrySettings.HairTipScale;
-		}
-	}
-	if ( !CardGenController->UpdateWidthMultipliers(HairWidths, RootScales, TipScales) )
-	{
-		UE_LOG(LogHairCardGenerator, Warning, TEXT("Strand width values inconsistent with number of physics groups, not using widths."));
-	}
-
 	// Run card group generation in loop
 	bool bAllGenerated = true;
 	for ( int i=0; i < GenerationSettings->GetFilterGroupSettings().Num(); ++i )
diff --git a/Engine/Plugins/Experimental/HairCardGenerator/Source/HairCardGeneratorEditor/Private/HairCardGeneratorPluginSettings.cpp b/Engine/Plugins/Experimental/HairCardGenerator/Source/HairCardGeneratorEditor/Private/HairCardGeneratorPluginSettings.cpp
--- a/Engine/Plugins/Experimental/HairCardGenerator/Source/HairCardGeneratorEditor/Private/HairCardGeneratorPluginSettings.cpp
+++ b/Engine/Plugins/Experimental/HairCardGenerator/Source/HairCardGeneratorEditor/Private/HairCardGeneratorPluginSettings.cpp
@@ -309,7 +309,7 @@ bool HairCardSettings_Converter::WriteSettingsToJson(UObject* SettingsObject, TS
 /* UHairCardGeneratorPluginSettings
  *****************************************************************************/
 
-void UHairCardGeneratorPluginSettings::SetSource(TObjectPtr<const UGroomAsset> InSourceObject, int32 GenLODIndex, int32 PhysGroupIndex)
+void UHairCardGeneratorPluginSettings::SetSource(TObjectPtr<UGroomAsset> InSourceObject, int32 GenLODIndex, int32 PhysGroupIndex)
 {
 	GroomAsset = InSourceObject;
 	LODIndex = GenLODIndex;
@@ -531,6 +531,7 @@ void UHairCardGeneratorPluginSettings::PostResetUpdates()
 {
 	UpdateOutputPaths();
 	UpdateParentInfo();
+	UpdateHairWidths();
 }
 
 void UHairCardGeneratorPluginSettings::UpdateOutputPaths()
@@ -559,6 +560,31 @@ void UHairCardGeneratorPluginSettings::UpdateParentInfo()
 		FilterGroupGenerationSettings[index]->UpdateParentName(BaseParentName, index);
 }
 
+void UHairCardGeneratorPluginSettings::UpdateHairWidths()
+{
+	int const NumPhysGroups = GroomAsset->GetHairDescriptionGroups().HairGroups.Num();
+
+	HairWidths.Init(-1.0, NumPhysGroups);
+	RootScales.Init(1.0, NumPhysGroups);
+	TipScales.Init(1.0, NumPhysGroups);
+	if ( bUseGroomAssetStrandWidth )
+	{
+		for ( int PhysGroupIdx = 0; PhysGroupIdx<NumPhysGroups; PhysGroupIdx++ )
+		{
+			const FHairGeometrySettings& GeometrySettings = GroomAsset->GetHairGroupsRendering()[PhysGroupIdx].GeometrySettings;
+
+			bOverrideHairWidth = GeometrySettings.HairWidth_Override;
+
+			if ( bOverrideHairWidth )
+			{
+				HairWidths[PhysGroupIdx] = GeometrySettings.HairWidth;
+			}
+			RootScales[PhysGroupIdx] = GeometrySettings.HairRootScale;
+			TipScales[PhysGroupIdx] = GeometrySettings.HairTipScale;
+		}
+	}
+}
+
 bool UHairCardGeneratorPluginSettings::FindDerivedTextureSettings()
 {
 	DerivedTextureSettingsName = TEXT("");
@@ -1000,6 +1026,10 @@ void UHairCardGeneratorPluginSettings::PostEditChangeProperty(FPropertyChangedEv
 	{
 		UpdateParentInfo();
 	}
+	else if (PropertyName == GET_MEMBER_NAME_CHECKED(UHairCardGeneratorPluginSettings, bUseGroomAssetStrandWidth))
+	{
+		UpdateHairWidths();
+	}
 }
 
 void UHairCardGeneratorPluginSettings::BuildDDCKey(FArchive& Ar)
diff --git a/Engine/Plugins/Experimental/HairCardGenerator/Source/HairCardGeneratorEditor/Private/HairCardGeneratorPluginSettings.h b/Engine/Plugins/Experimental/HairCardGenerator/Source/HairCardGeneratorEditor/Private/HairCardGeneratorPluginSettings.h
--- a/Engine/Plugins/Experimental/HairCardGenerator/Source/HairCardGeneratorEditor/Private/HairCardGeneratorPluginSettings.h
+++ b/Engine/Plugins/Experimental/HairCardGenerator/Source/HairCardGeneratorEditor/Private/HairCardGeneratorPluginSettings.h
@@ -128,7 +128,7 @@ public:
 	bool bGenerateGeometryForAllGroups = true;
 
 	// Seed value for pseudo-random number generation (set to a specific value for repeatable results)
-	UPROPERTY(EditAnywhere, BlueprintReadOnly, Category="Randomness")
+	UPROPERTY(EditAnywhere, BlueprintReadOnly, Category="Randomness", meta=(ClampMin="0", ClampMax="10000"))
 	int RandomSeed = 0;
 
 	// Place new card textures in reserved space from previous LOD
@@ -147,6 +147,18 @@ public:
 	UPROPERTY(EditAnywhere, BlueprintReadOnly, Category="Texture Rendering")
 	bool bUseGroomAssetStrandWidth = true;
 
+	UPROPERTY(BlueprintReadOnly, Category="Texture Rendering")
+	bool bOverrideHairWidth;
+
+	UPROPERTY(BlueprintReadOnly, Category="Texture Rendering")
+	TArray<float> HairWidths;
+
+	UPROPERTY(BlueprintReadOnly, Category="Texture Rendering")
+	TArray<float> RootScales;
+
+	UPROPERTY(BlueprintReadOnly, Category="Texture Rendering")
+	TArray<float> TipScales;
+
 	// Minimum strand depth value (mapped to 0 in depth texture)
 	UPROPERTY(BlueprintReadOnly, Category="Texture Rendering|Advanced", meta=(ClampMin="-10.0", ClampMax="0.0"))
 	float DepthMinimum = -2.0f;
@@ -162,7 +174,7 @@ public:
 	UPROPERTY(Transient, BlueprintReadOnly, Category="Asset")
 	FString OutputPath;
 
-	void SetSource(TObjectPtr<const UGroomAsset> SourceObject, int32 GenLODIndex = 0, int32 PhysGroupIndex = 0);
+	void SetSource(TObjectPtr<UGroomAsset> SourceObject, int32 GenLODIndex = 0, int32 PhysGroupIndex = 0);
 	// Load last-run values for options from card source description or defaults if no last-run info
 	void ResetFromSourceDescription(const FHairGroupsCardsSourceDescription& SourceDesc);
 	// Load settings from Json file if one exists
@@ -232,6 +244,7 @@ public:
 private:
 	void UpdateOutputPaths();
 	void UpdateParentInfo();
+	void UpdateHairWidths();
 	bool FindDerivedTextureSettings();
 	TSharedPtr<FJsonObject> GetParentTextureSettings() const;
 
@@ -286,7 +299,7 @@ private:
 	TArray<int> StrandFilterGroupIndexMap;
 
 	UPROPERTY(Transient)
-	TObjectPtr<const UGroomAsset> GroomAsset = nullptr;
+	TObjectPtr<UGroomAsset> GroomAsset = nullptr;
 
 	UPROPERTY(Transient)
 	TObjectPtr<const UStaticMesh> OldGeneratedMesh = nullptr;
