HairCardSolver: Handle hair card texture channel layouts during atlas generation.

diff --git a/Engine/Plugins/Experimental/HairCardGenerator/Config/Editor.ini b/Engine/Plugins/Experimental/HairCardGenerator/Config/Editor.ini
--- a/Engine/Plugins/Experimental/HairCardGenerator/Config/Editor.ini
+++ b/Engine/Plugins/Experimental/HairCardGenerator/Config/Editor.ini
@@ -1,4 +1,3 @@
 [/Script/HairCardGeneratorEditor.HairCardGeneratorEditorSettings]
 HairCardAssetNameFormat = "{groom_name}_LOD{lod}"
 HairCardAssetPathFormat = "{groom_dir}/{groom_name}_LOD{lod}"
-HairCardTextureNameFormats=((Depth, "{mesh_name}_depth"), (Tangent, "{mesh_name}_tangents"), (Attribute, "{mesh_name}_attributes"), (Coverage, "{mesh_name}_coverage"))
diff --git a/Engine/Plugins/Experimental/HairCardGenerator/Source/HairCardGeneratorEditor/Private/HairCardGenSettingsDetail.cpp b/Engine/Plugins/Experimental/HairCardGenerator/Source/HairCardGeneratorEditor/Private/HairCardGenSettingsDetail.cpp
--- a/Engine/Plugins/Experimental/HairCardGenerator/Source/HairCardGeneratorEditor/Private/HairCardGenSettingsDetail.cpp
+++ b/Engine/Plugins/Experimental/HairCardGenerator/Source/HairCardGeneratorEditor/Private/HairCardGenSettingsDetail.cpp
@@ -6,6 +6,7 @@
 #include "DetailCategoryBuilder.h"
 #include "DetailLayoutBuilder.h"
 #include "DetailWidgetRow.h"
+#include "Dom/JsonObject.h"
 #include "Widgets/Input/SCheckBox.h"
 #include "Widgets/Text/STextBlock.h"
 
@@ -50,7 +51,7 @@ void FHairCardSettingsDetailCustomization::CustomizeDetails(IDetailLayoutBuilder
         SNew(SCheckBox)
         .IsChecked(this, &FHairCardSettingsDetailCustomization::GetCheckValue, ReduceFromLODHandle)
         .OnCheckStateChanged(this, &FHairCardSettingsDetailCustomization::SetCheckValue, ReduceFromLODHandle)
-        .IsEnabled(this, &FHairCardSettingsDetailCustomization::CheckReduceFromLOD)
+        .IsEnabled(this, &FHairCardSettingsDetailCustomization::IsEnabledReduceFromLOD, ReduceFromLODHandle)
         .ToolTipText(this, &FHairCardSettingsDetailCustomization::ToolTipReduceFromLOD, ReduceFromLODHandle)
     ];
 
@@ -69,7 +70,7 @@ void FHairCardSettingsDetailCustomization::CustomizeDetails(IDetailLayoutBuilder
         SNew(SCheckBox)
         .IsChecked(this, &FHairCardSettingsDetailCustomization::GetCheckValue, UseReservedTxHandle)
         .OnCheckStateChanged(this, &FHairCardSettingsDetailCustomization::SetCheckValue, UseReservedTxHandle)
-        .IsEnabled(this, &FHairCardSettingsDetailCustomization::CheckUseReservedTx)
+        .IsEnabled(this, &FHairCardSettingsDetailCustomization::IsEnabledUseReservedTx, UseReservedTxHandle)
         .ToolTipText(this, &FHairCardSettingsDetailCustomization::ToolUseReservedTx, UseReservedTxHandle)
     ];
 }
@@ -101,91 +102,115 @@ void FHairCardSettingsDetailCustomization::SetCheckValue(ECheckBoxState NewState
     }
 }
 
-bool FHairCardSettingsDetailCustomization::CheckReduceFromLOD() const
+bool FHairCardSettingsDetailCustomization::IsEnabledReduceFromLOD(const TSharedPtr<IPropertyHandle> Property) const
 {
-    if ( !SettingsPtr.IsValid() )
-    {
-        return false;
-    }
+    FText Ignore;
+    return CheckReduceFromLOD(Property, Ignore);
+}
 
-    if ( SettingsPtr.Get()->GetLODIndex() < 1 )
-    {
-        return false;
-    }
-    
-    if ( !SettingsPtr.Get()->HasValidFullParent() )
-    {
-        return false;;
-    }
+FText FHairCardSettingsDetailCustomization::ToolTipReduceFromLOD(const TSharedPtr<IPropertyHandle> Property) const
+{
+    FText TooltipInfo;
+    CheckReduceFromLOD(Property, TooltipInfo);
 
-    return true;
+    return TooltipInfo;
 }
 
-FText FHairCardSettingsDetailCustomization::ToolTipReduceFromLOD(const TSharedPtr<IPropertyHandle> Property) const
+bool FHairCardSettingsDetailCustomization::CheckReduceFromLOD(const TSharedPtr<IPropertyHandle> Property, FText& OutTooltipInfo) const
 {
+    OutTooltipInfo = Property->GetToolTipText();
     if ( !SettingsPtr.IsValid() )
     {
-        return Property->GetToolTipText();
+        return false;
     }
 
-    if ( SettingsPtr.Get()->GetLODIndex() < 1 )
+    TObjectPtr<UHairCardGeneratorPluginSettings> SettingsPin = SettingsPtr.Get();
+    if ( SettingsPin->GetLODIndex() < 1 )
     {
-        return LOCTEXT("ReduceFromLOD.LOD0.ToolTip", "Cannot reduce LOD 0 (must run full generation)");
+        OutTooltipInfo = LOCTEXT("ReduceFromLOD.LOD0.ToolTip", "Cannot reduce LOD 0 (must run full generation)");
+        return false;
     }
 
-    if ( !SettingsPtr.Get()->HasValidFullParent() )
+    if ( !SettingsPin->ValidChannelLayouts() )
     {
-        return LOCTEXT("ReduceFromLOD.InvalidParent.ToolTip", "All lower LODs must be generated using the hair card generator tool");
+        OutTooltipInfo = LOCTEXT("ReduceFromLOD.InconsistentGroupLayouts.ToolTip", "Inconsistent texture layouts for groom groups at this LOD");
+        return false;
     }
-
-    return Property->GetToolTipText();
-}
-
-
-bool FHairCardSettingsDetailCustomization::CheckUseReservedTx() const
-{
-    if ( !SettingsPtr.IsValid() )
+    
+    TSharedPtr<FJsonObject> ParentSettingsJson = SettingsPin->GetFullParent();
+    if ( !ParentSettingsJson.IsValid() )
     {
+        OutTooltipInfo = LOCTEXT("ReduceFromLOD.InvalidParent.ToolTip", "All lower LODs must be generated using the hair card generator tool");
         return false;
     }
 
-    if ( SettingsPtr->bReduceCardsFromPreviousLOD )
+    if ( !ParentSettingsJson->HasTypedField<EJson::String>(TEXT("ChannelLayout")) )
     {
+        OutTooltipInfo = LOCTEXT("ReduceFromLOD.InvalidChannelLayout.ToolTip", "Invalid texture layout setting in previous LOD");
         return false;
     }
 
-    if ( !SettingsPtr->HasDerivedTextureSettings() )
+    UEnum* EnumClass = StaticEnum<EHairTextureLayout>();
+    if ( !EnumClass || ParentSettingsJson->GetStringField(TEXT("ChannelLayout")) != EnumClass->GetNameStringByValue((int64)SettingsPin->GetChannelLayout()) )
     {
+        OutTooltipInfo = LOCTEXT("ReduceFromLOD.InconsistentParentLayout.ToolTip", "Parent texture layout setting differs from current LOD texture layout");
         return false;
     }
 
-    // TODO: Handle limiting/reserving texture by using texture resolution to compute reserved space in pixels
-    return (SettingsPtr->GetDerivedReservedTextureSize() >= 5);
+    return true;
 }
 
+bool FHairCardSettingsDetailCustomization::IsEnabledUseReservedTx(const TSharedPtr<IPropertyHandle> Property) const
+{
+    FText Ignore;
+    return CheckUseReservedTx(Property, Ignore);
+}
+
+
 FText FHairCardSettingsDetailCustomization::ToolUseReservedTx(const TSharedPtr<IPropertyHandle> Property) const
 {
+    FText TooltipInfo;
+    CheckUseReservedTx(Property, TooltipInfo);
+
+    return TooltipInfo;
+}
+
+bool FHairCardSettingsDetailCustomization::CheckUseReservedTx(const TSharedPtr<IPropertyHandle> Property, FText& OutTooltipInfo) const
+{
+    OutTooltipInfo = Property->GetToolTipText();
     if ( !SettingsPtr.IsValid() )
     {
-        return Property->GetToolTipText();
+        return false;
     }
 
-    if ( SettingsPtr->bReduceCardsFromPreviousLOD )
+    TObjectPtr<UHairCardGeneratorPluginSettings> SettingsPin = SettingsPtr.Get();
+    if ( SettingsPin->bReduceCardsFromPreviousLOD )
     {
-        return LOCTEXT("UseReservedTx.Reducing.ToolTip", "Reduced card geometry will use previous LOD texture UVs");
+        OutTooltipInfo = LOCTEXT("UseReservedTx.Reducing.ToolTip", "Reduced card geometry will use previous LOD texture UVs");
+        return false;
     }
 
-    if ( !SettingsPtr->HasDerivedTextureSettings() )
+    if ( !SettingsPin->HasDerivedTextureSettings() )
     {
-        return LOCTEXT("UseReservedTx.InvalidParent.ToolTip", "Lower LODs must be generated using hair card generator tool with reserved space");
+        OutTooltipInfo = LOCTEXT("UseReservedTx.InvalidParent.ToolTip", "Lower LODs must be generated using hair card generator tool with reserved space");
+        return false;
     }
 
-    if ( SettingsPtr->GetDerivedReservedTextureSize() < 5 )
+    // TODO: Handle limiting/reserving texture by using texture resolution to compute reserved space in pixels
+    if ( SettingsPin->GetDerivedReservedTextureSize() < 5 )
     {
-        return LOCTEXT("UseReservedTx.NoReservedSpace.ToolTip", "No space reserved in parent LOD chain");
+        OutTooltipInfo = LOCTEXT("UseReservedTx.NoReservedSpace.ToolTip", "No space reserved in parent LOD chain");
+        return false;
     }
 
-    return Property->GetToolTipText();
+    UEnum* EnumClass = StaticEnum<EHairTextureLayout>();
+    if ( !EnumClass || SettingsPin->GetDerivedTextureChannelLayout() != EnumClass->GetNameStringByValue((int64)SettingsPin->GetChannelLayout()) )
+    {
+        OutTooltipInfo = LOCTEXT("UseReservedTx.InconsistentReservedLayout.ToolTip", "Parent reserved texture layout setting differs from current LOD texture layout");
+        return false;
+    }
+    
+    return true;
 }
 
 
diff --git a/Engine/Plugins/Experimental/HairCardGenerator/Source/HairCardGeneratorEditor/Private/HairCardGenSettingsDetail.h b/Engine/Plugins/Experimental/HairCardGenerator/Source/HairCardGeneratorEditor/Private/HairCardGenSettingsDetail.h
--- a/Engine/Plugins/Experimental/HairCardGenerator/Source/HairCardGeneratorEditor/Private/HairCardGenSettingsDetail.h
+++ b/Engine/Plugins/Experimental/HairCardGenerator/Source/HairCardGeneratorEditor/Private/HairCardGenSettingsDetail.h
@@ -26,11 +26,13 @@ private:
     ECheckBoxState GetCheckValue(const TSharedPtr<IPropertyHandle> Property) const;
     void SetCheckValue(ECheckBoxState NewState, const TSharedPtr<IPropertyHandle> Property) const;
 
-    bool CheckReduceFromLOD() const;
+    bool IsEnabledReduceFromLOD(const TSharedPtr<IPropertyHandle> Property) const;
     FText ToolTipReduceFromLOD(const TSharedPtr<IPropertyHandle> Property) const;
+    bool CheckReduceFromLOD(const TSharedPtr<IPropertyHandle> Property, FText& OutTooltipInfo) const;
 
-    bool CheckUseReservedTx() const;
+    bool IsEnabledUseReservedTx(const TSharedPtr<IPropertyHandle> Property) const;
     FText ToolUseReservedTx(const TSharedPtr<IPropertyHandle> Property) const;
+    bool CheckUseReservedTx(const TSharedPtr<IPropertyHandle> Property, FText& OutTooltipInfo) const;
 
     TWeakObjectPtr<UHairCardGeneratorPluginSettings> SettingsPtr;
 };
diff --git a/Engine/Plugins/Experimental/HairCardGenerator/Source/HairCardGeneratorEditor/Private/HairCardGeneratorEditorModule.cpp b/Engine/Plugins/Experimental/HairCardGenerator/Source/HairCardGeneratorEditor/Private/HairCardGeneratorEditorModule.cpp
--- a/Engine/Plugins/Experimental/HairCardGenerator/Source/HairCardGeneratorEditor/Private/HairCardGeneratorEditorModule.cpp
+++ b/Engine/Plugins/Experimental/HairCardGenerator/Source/HairCardGeneratorEditor/Private/HairCardGeneratorEditorModule.cpp
@@ -17,6 +17,7 @@
 #include "PropertyEditorModule.h"
 #include "AutomatedAssetImportData.h"
 #include "GroomAsset.h"
+#include "GroomAssetCards.h" // for GetHairTextureLayoutTextureCount
 #include "GroomCacheData.h" // for EGroomBasisType, EGroomCurveType
 #include "HairDescription.h"
 #include "ObjectTools.h"
@@ -43,7 +44,7 @@ LLM_DEFINE_TAG(HairCardGenerator);
 struct FHairCardImportSettings
 {
 	FString MeshName;
-	TMap<EHairAtlasTextureType, FString> TextureNames;
+	TArray<FString> TextureNames;
 };
 
 
@@ -458,29 +459,15 @@ static bool HairCardGeneratorEditor_Impl::GenerateCardsForCardGroup(TObjectPtr<c
 
 static void HairCardGeneratorEditor_Impl::GetTextureImportNames(UHairCardGeneratorPluginSettings* Settings, FHairCardImportSettings& ImportSettings)
 {
+	uint32 NumTextures = GetHairTextureLayoutTextureCount(Settings->ChannelLayout);
 	FString TextureBasename = Settings->GetTextureImportBaseName();
 
-	ImportSettings.TextureNames.Empty();
+	ImportSettings.TextureNames.SetNum(NumTextures);
 	ImportSettings.MeshName = Settings->BaseFilename;
 
-	FFormatNamedArguments FilenameFormatArgs;
-	FilenameFormatArgs.Add(TEXT("groom_name"), FText::FromString(Settings->GetGroomName()));
-	FilenameFormatArgs.Add(TEXT("lod"), Settings->GetLODIndex());
-	FilenameFormatArgs.Add(TEXT("mesh_name"), FText::FromString(TextureBasename));
-
-	const UHairCardGeneratorEditorSettings* ConfigSettings = GetDefault<UHairCardGeneratorEditorSettings>();
-	ImportSettings.TextureNames.Reserve(ConfigSettings->HairCardTextureNameFormats.Num());
-	for (auto TextureNameIt : ConfigSettings->HairCardTextureNameFormats)
+	for (uint32 i = 0; i < NumTextures; ++i)
 	{
-		FString FormattedTextureName = FText::Format(FText::FromString(TextureNameIt.Value), FilenameFormatArgs).ToString();
-		if (FPaths::ValidatePath(FormattedTextureName))
-		{
-			ImportSettings.TextureNames.Add(TextureNameIt.Key, FormattedTextureName);
-		}
-		else
-		{
-			UE_LOG(LogHairCardGenerator, Warning, TEXT("Invalid texture name config: '%s'"), *TextureNameIt.Value);
-		}
+		ImportSettings.TextureNames[i] = FString::Printf(TEXT("%s_TS%d"), *TextureBasename, i);
 	}
 }
 
@@ -496,29 +483,18 @@ static TArray<UObject*> HairCardGeneratorEditor_Impl::ImportHairCardTextures(con
 	AutoImportLinearTextures->bSkipReadOnly = false;
 	AutoImportLinearTextures->Factory = LinearTextureFactory;
 
-	UAutomatedAssetImportData* AutomatedImportData = NewObject<UAutomatedAssetImportData>();
-	AutomatedImportData->DestinationPath = ImportDst;
-	AutomatedImportData->bReplaceExisting = true;
-	AutomatedImportData->bSkipReadOnly = false;
-
 	IPlatformFile& FileManager = FPlatformFileManager::Get().GetPlatformFile();
-	FileManager.IterateDirectory(*ImportSrc, [AutomatedImportData, AutoImportLinearTextures, ImportSettings](const TCHAR* FilenameOrDirectory, const bool bIsDirectory)->bool
+	FileManager.IterateDirectory(*ImportSrc, [AutoImportLinearTextures, ImportSettings](const TCHAR* FilenameOrDirectory, const bool bIsDirectory)->bool
 		{
 			if (!bIsDirectory)
 			{
-				FString basename = FPaths::GetBaseFilename(FilenameOrDirectory, true);
-				const EHairAtlasTextureType* pType = ImportSettings.TextureNames.FindKey(basename);
-				if ( pType && (*pType == EHairAtlasTextureType::Tangent || *pType == EHairAtlasTextureType::Attribute ) )
-					AutoImportLinearTextures->Filenames.Add(FString(FilenameOrDirectory));
-				else
-					AutomatedImportData->Filenames.Add(FilenameOrDirectory);
+				AutoImportLinearTextures->Filenames.Add(FString(FilenameOrDirectory));
 			}
 			return true;
 		}
 	);
 
-	TArray<UObject*> ImportedAssets = FAssetToolsModule::GetModule().Get().ImportAssetsAutomated(AutomatedImportData);
-	ImportedAssets.Append(FAssetToolsModule::GetModule().Get().ImportAssetsAutomated(AutoImportLinearTextures));
+	TArray<UObject*> ImportedAssets = FAssetToolsModule::GetModule().Get().ImportAssetsAutomated(AutoImportLinearTextures);
 	return ImportedAssets;
 }
 
@@ -570,15 +546,12 @@ static int32 HairCardGeneratorEditor_Impl::AssignImportedTextures(const TArray<U
 		if ( TextureImport )
 		{
 			bool bFound = false;
-			for (auto NameTemplateIt : ImportSettings.TextureNames)
+			int SlotIdx = ImportSettings.TextureNames.Find(TextureImport->GetName());
+			if (SlotIdx != INDEX_NONE)
 			{
-				if (TextureImport->GetName().Equals(NameTemplateIt.Value))
-				{
-					Assignee.Textures.SetTexture(NameTemplateIt.Key, TextureImport);
-					bFound = true;
-					++AssignedAssets;
-					break;
-				}
+				Assignee.Textures.SetTexture(SlotIdx, TextureImport);
+				bFound = true;
+				++AssignedAssets;
 			}
 
 			UE_CLOG(!bFound, LogHairCardGenerator, Warning, TEXT("Unable to determine which slot imported texture ('%s') is for. Check your config."), *TextureImport->GetName());
@@ -698,6 +671,9 @@ bool FHairCardGeneratorEditorModule::GenerateHairCardsForLOD(UGroomAsset* NewGro
 	// Write the settings that we will attempt to output to card desc (won't be stored on failure see FGroomRenderingDetails::OnGenerateCardDataUsingPlugin for details)
 	CardsDesc.GenerationSettings = GenerationSettings;
 
+	// Need to set the CardsDesc texture layout appropriately in case we succeed
+	CardsDesc.Textures.SetLayout(GenerationSettings->ChannelLayout);
+
 	// TODO: On error will force full regen (is this the best we can do?)
 	GenerationSettings->ClearPipelineGenerated();
 
diff --git a/Engine/Plugins/Experimental/HairCardGenerator/Source/HairCardGeneratorEditor/Private/HairCardGeneratorEditorSettings.h b/Engine/Plugins/Experimental/HairCardGenerator/Source/HairCardGeneratorEditor/Private/HairCardGeneratorEditorSettings.h
--- a/Engine/Plugins/Experimental/HairCardGenerator/Source/HairCardGeneratorEditor/Private/HairCardGeneratorEditorSettings.h
+++ b/Engine/Plugins/Experimental/HairCardGenerator/Source/HairCardGeneratorEditor/Private/HairCardGeneratorEditorSettings.h
@@ -23,7 +23,4 @@ public:
 
 	UPROPERTY(Config)
 	FString HairCardAssetPathFormat;
-
-	UPROPERTY(Config)
-	TMap<EHairAtlasTextureType, FString> HairCardTextureNameFormats;
 };
\ No newline at end of file
diff --git a/Engine/Plugins/Experimental/HairCardGenerator/Source/HairCardGeneratorEditor/Private/HairCardGeneratorPluginSettings.cpp b/Engine/Plugins/Experimental/HairCardGenerator/Source/HairCardGeneratorEditor/Private/HairCardGeneratorPluginSettings.cpp
--- a/Engine/Plugins/Experimental/HairCardGenerator/Source/HairCardGeneratorEditor/Private/HairCardGeneratorPluginSettings.cpp
+++ b/Engine/Plugins/Experimental/HairCardGenerator/Source/HairCardGeneratorEditor/Private/HairCardGeneratorPluginSettings.cpp
@@ -532,6 +532,7 @@ void UHairCardGeneratorPluginSettings::PostResetUpdates()
 	UpdateOutputPaths();
 	UpdateParentInfo();
 	UpdateHairWidths();
+	UpdateChannelLayout();
 }
 
 void UHairCardGeneratorPluginSettings::UpdateOutputPaths()
@@ -542,6 +543,38 @@ void UHairCardGeneratorPluginSettings::UpdateOutputPaths()
 		FilterGroupGenerationSettings[index]->UpdateGenerateFilename(BaseFilename, index);
 }
 
+void UHairCardGeneratorPluginSettings::UpdateChannelLayout()
+{
+	ChannelLayout = EHairTextureLayout::Layout0;
+
+	bool bFoundDesc = false;
+	bool bAllMatched = true;
+	EHairTextureLayout CheckLayout =  EHairTextureLayout::Layout0;
+	for (const FHairGroupsCardsSourceDescription& Desc : GroomAsset->GetHairGroupsCards())
+	{
+		if (Desc.LODIndex != LODIndex)
+		{
+			continue;
+		}
+
+		if (!bFoundDesc)
+		{
+			bFoundDesc = true;
+			CheckLayout = Desc.Textures.Layout;
+		}
+
+		if (Desc.GroupIndex == GenerateForGroomGroup)
+		{
+			ChannelLayout = Desc.Textures.Layout;
+		}
+
+		bAllMatched = bAllMatched & (CheckLayout == Desc.Textures.Layout);
+		CheckLayout = Desc.Textures.Layout;
+	}
+
+	bValidChannelLayouts = bFoundDesc & bAllMatched;
+}
+
 void UHairCardGeneratorPluginSettings::UpdateParentInfo()
 {
 	BaseParentName = TEXT("");
@@ -687,7 +720,7 @@ FString UHairCardGeneratorPluginSettings::GetTextureContentPath() const
 }
 
 
-bool UHairCardGeneratorPluginSettings::HasValidFullParent() const
+TSharedPtr<FJsonObject> UHairCardGeneratorPluginSettings::GetFullParent() const
 {
 	TSharedPtr<FJsonObject> DerivedFullSettings = TraverseSettingsParents(BaseParentName, 
 		[](const TSharedPtr<FJsonObject>& SettingsRoot)
@@ -697,7 +730,7 @@ bool UHairCardGeneratorPluginSettings::HasValidFullParent() const
 		}
 	);
 
-	return DerivedFullSettings.IsValid();
+	return DerivedFullSettings;
 }
 
 bool UHairCardGeneratorPluginSettings::HasDerivedTextureSettings() const
@@ -718,6 +751,17 @@ int UHairCardGeneratorPluginSettings::GetDerivedReservedTextureSize() const
 	return FJsonHelper::GetIntFieldDefault(ParentTextureSettings, TEXT("ReserveTextureSpaceLOD"));
 }
 
+FString UHairCardGeneratorPluginSettings::GetDerivedTextureChannelLayout() const
+{
+	TSharedPtr<FJsonObject> ParentTextureSettings = GetParentTextureSettings();
+	if ( !ParentTextureSettings.IsValid() )
+	{
+		return TEXT("");
+	}
+
+	return FJsonHelper::GetStringFieldDefault(ParentTextureSettings, TEXT("ChannelLayout"));
+}
+
 
 FString UHairCardGeneratorPluginSettings::GetSettingsFileFromBase(const FString& SettingsName)
 {
@@ -837,6 +881,7 @@ void UHairCardGeneratorPluginSettings::SerializeEditableSettings(FArchive& Ar)
 	Ar << bUseReservedSpaceFromPreviousLOD;
 	Ar << ReserveTextureSpaceLOD;
 	Ar << bUseGroomAssetStrandWidth;
+	Ar << ChannelLayout;
 }
 
 void UHairCardGeneratorPluginSettings::ResetNumFilterGroups(int Count)
diff --git a/Engine/Plugins/Experimental/HairCardGenerator/Source/HairCardGeneratorEditor/Private/HairCardGeneratorPluginSettings.h b/Engine/Plugins/Experimental/HairCardGenerator/Source/HairCardGeneratorEditor/Private/HairCardGeneratorPluginSettings.h
--- a/Engine/Plugins/Experimental/HairCardGenerator/Source/HairCardGeneratorEditor/Private/HairCardGeneratorPluginSettings.h
+++ b/Engine/Plugins/Experimental/HairCardGenerator/Source/HairCardGeneratorEditor/Private/HairCardGeneratorPluginSettings.h
@@ -139,6 +139,10 @@ public:
 	UPROPERTY(EditAnywhere, BlueprintReadOnly, Category="Texture Rendering")
 	EHairCardAtlasSize AtlasSize = EHairCardAtlasSize::AtlasSize4096;
 
+	// Texture layout selected (NOTE: This automatically pulls from group 0)
+	UPROPERTY(BlueprintReadOnly, Category="Texture Rendering")
+	EHairTextureLayout ChannelLayout = EHairTextureLayout::Layout0;
+
 	// Percentage of texture atlas space to reserve for higher LODs 
 	UPROPERTY(EditAnywhere, BlueprintReadOnly, Category="Texture Rendering", meta=(ClampMin="0", ClampMax="75"))
 	int ReserveTextureSpaceLOD = 0;
@@ -223,15 +227,18 @@ public:
 	FString GetTextureImportPath() const;
 	FString GetTextureContentPath() const;
 
-	bool HasValidFullParent() const;
+	TSharedPtr<FJsonObject> GetFullParent() const;
 	bool HasDerivedTextureSettings() const;
 	int GetDerivedReservedTextureSize() const;
+	FString GetDerivedTextureChannelLayout() const;
 	
 	FString GetGroomName() const;
 
 	int32 GetLODIndex() const { return LODIndex; }
 	int32 GetGenerateForGroomGroup() const { return GenerateForGroomGroup; }
 	void SetGenerateForGroomGroup(int GroupIndex) { GenerateForGroomGroup = GroupIndex; }
+	EHairTextureLayout GetChannelLayout() const { return ChannelLayout; }
+	bool ValidChannelLayouts() { return bValidChannelLayouts; }
 
 	TArray<TObjectPtr<UHairCardGeneratorGroupSettings>>& GetFilterGroupSettings() { return FilterGroupGenerationSettings; }
 	const TArray<TObjectPtr<UHairCardGeneratorGroupSettings>>& GetFilterGroupSettings() const { return FilterGroupGenerationSettings; }
@@ -243,6 +250,9 @@ public:
 
 private:
 	void UpdateOutputPaths();
+
+	void UpdateChannelLayout();
+
 	void UpdateParentInfo();
 	void UpdateHairWidths();
 	bool FindDerivedTextureSettings();
@@ -298,6 +308,10 @@ private:
 	UPROPERTY(Transient,BlueprintReadOnly, Category="Settings Groups", meta=(AllowPrivateAccess))
 	TArray<int> StrandFilterGroupIndexMap;
 
+	// All group texture layouts are the same (required for current texture generation)
+	UPROPERTY(Transient)
+	bool bValidChannelLayouts;
+
 	UPROPERTY(Transient)
 	TObjectPtr<UGroomAsset> GroomAsset = nullptr;
 
diff --git a/Engine/Plugins/Runtime/HairStrands/Source/HairStrandsCore/Private/GroomAssetCards.cpp b/Engine/Plugins/Runtime/HairStrands/Source/HairStrandsCore/Private/GroomAssetCards.cpp
--- a/Engine/Plugins/Runtime/HairStrands/Source/HairStrandsCore/Private/GroomAssetCards.cpp
+++ b/Engine/Plugins/Runtime/HairStrands/Source/HairStrandsCore/Private/GroomAssetCards.cpp
@@ -16,19 +16,18 @@ FHairGroupsCardsSourceDescription::FHairGroupsCardsSourceDescription()
 	LODIndex = -1;
 }
 
-void FHairGroupCardsTextures::SetTexture(EHairAtlasTextureType SlotID, UTexture2D* Texture)
+void FHairGroupCardsTextures::SetLayout(EHairTextureLayout InLayout)
 {
-	Layout = EHairTextureLayout::Layout0; // Default hair card layout for now
-	Textures.SetNum(6);
-	switch (SlotID)
-	{
-	case EHairAtlasTextureType::Depth: 			Textures[0] = Texture; break;
-	case EHairAtlasTextureType::Coverage:		Textures[1] = Texture; break;
-	case EHairAtlasTextureType::Tangent:		Textures[2] = Texture; break;
-	case EHairAtlasTextureType::Attribute:		Textures[3] = Texture; break;
-	case EHairAtlasTextureType::AuxilaryData:	Textures[4] = Texture; break;
-	case EHairAtlasTextureType::Material:		Textures[5] = Texture; break;
-	};
+	Layout = InLayout;
+
+	uint32 TextureCount = GetHairTextureLayoutTextureCount(Layout);
+	Textures.SetNum(TextureCount, true);
+}
+
+void FHairGroupCardsTextures::SetTexture(int32 SlotIdx, UTexture2D* Texture)
+{
+	check(SlotIdx < Textures.Num());
+	Textures[SlotIdx] = Texture;
 }
 
 bool FHairGroupsCardsSourceDescription::operator==(const FHairGroupsCardsSourceDescription& A) const
diff --git a/Engine/Plugins/Runtime/HairStrands/Source/HairStrandsCore/Public/GroomAssetCards.h b/Engine/Plugins/Runtime/HairStrands/Source/HairStrandsCore/Public/GroomAssetCards.h
--- a/Engine/Plugins/Runtime/HairStrands/Source/HairStrandsCore/Public/GroomAssetCards.h
+++ b/Engine/Plugins/Runtime/HairStrands/Source/HairStrandsCore/Public/GroomAssetCards.h
@@ -82,7 +82,8 @@ struct HAIRSTRANDSCORE_API FHairGroupCardsTextures
 	UPROPERTY()
 	TObjectPtr<UTexture2D> MaterialTexture_DEPRECATED = nullptr;
 
-	void SetTexture(EHairAtlasTextureType SlotID, UTexture2D* Texture);
+	void SetLayout(EHairTextureLayout InLayout);
+	void SetTexture(int32 SlotID, UTexture2D* Texture);
 
 	bool bNeedToBeSaved = false;
 };
