HairCardSolver: Fixup texture generation/UI for multi-LOD handling to be less error-prone and less confusing.

diff --git a/Engine/Plugins/Experimental/HairCardGenerator/Source/HairCardGeneratorEditor/Private/HairCardGenSettingsDetail.cpp b/Engine/Plugins/Experimental/HairCardGenerator/Source/HairCardGeneratorEditor/Private/HairCardGenSettingsDetail.cpp
--- a/Engine/Plugins/Experimental/HairCardGenerator/Source/HairCardGeneratorEditor/Private/HairCardGenSettingsDetail.cpp
+++ b/Engine/Plugins/Experimental/HairCardGenerator/Source/HairCardGeneratorEditor/Private/HairCardGenSettingsDetail.cpp
@@ -124,40 +124,7 @@ bool FHairCardSettingsDetailCustomization::CheckReduceFromLOD(const TSharedPtr<I
         return false;
     }
 
-    TObjectPtr<UHairCardGeneratorPluginSettings> SettingsPin = SettingsPtr.Get();
-    if ( SettingsPin->GetLODIndex() < 1 )
-    {
-        OutTooltipInfo = LOCTEXT("ReduceFromLOD.LOD0.ToolTip", "Cannot reduce LOD 0 (must run full generation)");
-        return false;
-    }
-
-    if ( !SettingsPin->ValidChannelLayouts() )
-    {
-        OutTooltipInfo = LOCTEXT("ReduceFromLOD.InconsistentGroupLayouts.ToolTip", "Inconsistent texture layouts for groom groups at this LOD");
-        return false;
-    }
-    
-    TSharedPtr<FJsonObject> ParentSettingsJson = SettingsPin->GetFullParent();
-    if ( !ParentSettingsJson.IsValid() )
-    {
-        OutTooltipInfo = LOCTEXT("ReduceFromLOD.InvalidParent.ToolTip", "All lower LODs must be generated using the hair card generator tool");
-        return false;
-    }
-
-    if ( !ParentSettingsJson->HasTypedField<EJson::String>(TEXT("ChannelLayout")) )
-    {
-        OutTooltipInfo = LOCTEXT("ReduceFromLOD.InvalidChannelLayout.ToolTip", "Invalid texture layout setting in previous LOD");
-        return false;
-    }
-
-    UEnum* EnumClass = StaticEnum<EHairTextureLayout>();
-    if ( !EnumClass || ParentSettingsJson->GetStringField(TEXT("ChannelLayout")) != EnumClass->GetNameStringByValue((int64)SettingsPin->GetChannelLayout()) )
-    {
-        OutTooltipInfo = LOCTEXT("ReduceFromLOD.InconsistentParentLayout.ToolTip", "Parent texture layout setting differs from current LOD texture layout");
-        return false;
-    }
-
-    return true;
+    return SettingsPtr.Get()->CanReduceFromLOD(&OutTooltipInfo);
 }
 
 bool FHairCardSettingsDetailCustomization::IsEnabledUseReservedTx(const TSharedPtr<IPropertyHandle> Property) const
@@ -183,34 +150,7 @@ bool FHairCardSettingsDetailCustomization::CheckUseReservedTx(const TSharedPtr<I
         return false;
     }
 
-    TObjectPtr<UHairCardGeneratorPluginSettings> SettingsPin = SettingsPtr.Get();
-    if ( SettingsPin->bReduceCardsFromPreviousLOD )
-    {
-        OutTooltipInfo = LOCTEXT("UseReservedTx.Reducing.ToolTip", "Reduced card geometry will use previous LOD texture UVs");
-        return false;
-    }
-
-    if ( !SettingsPin->HasDerivedTextureSettings() )
-    {
-        OutTooltipInfo = LOCTEXT("UseReservedTx.InvalidParent.ToolTip", "Lower LODs must be generated using hair card generator tool with reserved space");
-        return false;
-    }
-
-    // TODO: Handle limiting/reserving texture by using texture resolution to compute reserved space in pixels
-    if ( SettingsPin->GetDerivedReservedTextureSize() < 5 )
-    {
-        OutTooltipInfo = LOCTEXT("UseReservedTx.NoReservedSpace.ToolTip", "No space reserved in parent LOD chain");
-        return false;
-    }
-
-    UEnum* EnumClass = StaticEnum<EHairTextureLayout>();
-    if ( !EnumClass || SettingsPin->GetDerivedTextureChannelLayout() != EnumClass->GetNameStringByValue((int64)SettingsPin->GetChannelLayout()) )
-    {
-        OutTooltipInfo = LOCTEXT("UseReservedTx.InconsistentReservedLayout.ToolTip", "Parent reserved texture layout setting differs from current LOD texture layout");
-        return false;
-    }
-    
-    return true;
+    return SettingsPtr.Get()->CanUseReservedTx(&OutTooltipInfo);
 }
 
 
diff --git a/Engine/Plugins/Experimental/HairCardGenerator/Source/HairCardGeneratorEditor/Private/HairCardGeneratorEditorModule.cpp b/Engine/Plugins/Experimental/HairCardGenerator/Source/HairCardGeneratorEditor/Private/HairCardGeneratorEditorModule.cpp
--- a/Engine/Plugins/Experimental/HairCardGenerator/Source/HairCardGeneratorEditor/Private/HairCardGeneratorEditorModule.cpp
+++ b/Engine/Plugins/Experimental/HairCardGenerator/Source/HairCardGeneratorEditor/Private/HairCardGeneratorEditorModule.cpp
@@ -197,7 +197,6 @@ static bool HairCardGeneratorEditor_Impl::LoadGroomData(UGroomAsset* NewGroomAss
 {
 	if (HairCardGeneratorEditor_Impl::LastGroom == NewGroomAsset && !SaveCached)
 	{
-		UE_LOG(LogHairCardGenerator, Warning, TEXT("Groom already loaded."));
 		return true;
 	}
 	else
diff --git a/Engine/Plugins/Experimental/HairCardGenerator/Source/HairCardGeneratorEditor/Private/HairCardGeneratorPluginSettings.cpp b/Engine/Plugins/Experimental/HairCardGenerator/Source/HairCardGeneratorEditor/Private/HairCardGeneratorPluginSettings.cpp
--- a/Engine/Plugins/Experimental/HairCardGenerator/Source/HairCardGeneratorEditor/Private/HairCardGeneratorPluginSettings.cpp
+++ b/Engine/Plugins/Experimental/HairCardGenerator/Source/HairCardGeneratorEditor/Private/HairCardGeneratorPluginSettings.cpp
@@ -24,6 +24,8 @@
 #include "Serialization/MemoryWriter.h"
 
 
+#define LOCTEXT_NAMESPACE "HairCardPluginSettings"
+
 /* Json helper functions
  *****************************************************************************/
 namespace FJsonHelper
@@ -533,6 +535,8 @@ void UHairCardGeneratorPluginSettings::PostResetUpdates()
 	UpdateParentInfo();
 	UpdateHairWidths();
 	UpdateChannelLayout();
+
+	EnforceValidLODSettings();
 }
 
 void UHairCardGeneratorPluginSettings::UpdateOutputPaths()
@@ -618,6 +622,15 @@ void UHairCardGeneratorPluginSettings::UpdateHairWidths()
 	}
 }
 
+void UHairCardGeneratorPluginSettings::EnforceValidLODSettings()
+{
+	// Force these settings to be valid in case reloading settings changed something
+	// TODO: Rework all of this to use a more robust json settings/completion framework
+	bReduceCardsFromPreviousLOD &= CanReduceFromLOD();
+	bUseReservedSpaceFromPreviousLOD &= CanReduceFromLOD();
+}
+
+
 bool UHairCardGeneratorPluginSettings::FindDerivedTextureSettings()
 {
 	DerivedTextureSettingsName = TEXT("");
@@ -882,6 +895,10 @@ void UHairCardGeneratorPluginSettings::SerializeEditableSettings(FArchive& Ar)
 	Ar << ReserveTextureSpaceLOD;
 	Ar << bUseGroomAssetStrandWidth;
 	Ar << ChannelLayout;
+	// TODO: Accurate diff requires conditional for these settings
+	Ar << HairWidths;
+	Ar << RootScales;
+	Ar << TipScales;
 }
 
 void UHairCardGeneratorPluginSettings::ResetNumFilterGroups(int Count)
@@ -1099,6 +1116,83 @@ FString UHairCardGeneratorPluginSettings::GetGroomName() const
 	return GroomAsset->GetName();
 }
 
+template <typename T>
+void SetIfValid(T* OutPtr, const T& Val)
+{
+	if (OutPtr)
+	{
+		*OutPtr = Val;
+	}
+}
+
+bool UHairCardGeneratorPluginSettings::CanReduceFromLOD(FText* OutInvalidInfo) const
+{
+    if ( GetLODIndex() < 1 )
+    {
+		SetIfValid(OutInvalidInfo, LOCTEXT("ReduceFromLOD.LOD0.ToolTip", "Cannot reduce LOD 0 (must run full generation)"));
+        return false;
+    }
+
+    if ( !ValidChannelLayouts() )
+    {
+        SetIfValid(OutInvalidInfo, LOCTEXT("ReduceFromLOD.InconsistentGroupLayouts.ToolTip", "Inconsistent texture layouts for groom groups at this LOD"));
+        return false;
+    }
+    
+    TSharedPtr<FJsonObject> ParentSettingsJson = GetFullParent();
+    if ( !ParentSettingsJson.IsValid() )
+    {
+        SetIfValid(OutInvalidInfo, LOCTEXT("ReduceFromLOD.InvalidParent.ToolTip", "All lower LODs must be generated using the hair card generator tool"));
+        return false;
+    }
+
+    if ( !ParentSettingsJson->HasTypedField<EJson::String>(TEXT("ChannelLayout")) )
+    {
+        SetIfValid(OutInvalidInfo, LOCTEXT("ReduceFromLOD.InvalidChannelLayout.ToolTip", "Invalid texture layout setting in previous LOD"));
+        return false;
+    }
+
+    UEnum* EnumClass = StaticEnum<EHairTextureLayout>();
+    if ( !EnumClass || ParentSettingsJson->GetStringField(TEXT("ChannelLayout")) != EnumClass->GetNameStringByValue((int64)GetChannelLayout()) )
+    {
+		SetIfValid(OutInvalidInfo, LOCTEXT("ReduceFromLOD.InconsistentParentLayout.ToolTip", "Parent texture layout setting differs from current LOD texture layout"));
+        return false;
+    }
+
+    return true;
+}
+
+bool UHairCardGeneratorPluginSettings::CanUseReservedTx(FText* OutInvalidInfo) const
+{
+    if ( bReduceCardsFromPreviousLOD )
+    {
+        SetIfValid(OutInvalidInfo, LOCTEXT("UseReservedTx.Reducing.ToolTip", "Reduced card geometry will use previous LOD texture UVs"));
+        return false;
+    }
+
+    if ( !HasDerivedTextureSettings() )
+    {
+        SetIfValid(OutInvalidInfo, LOCTEXT("UseReservedTx.InvalidParent.ToolTip", "Lower LODs must be generated using hair card generator tool with reserved space"));
+        return false;
+    }
+
+    // TODO: Handle limiting/reserving texture by using texture resolution to compute reserved space in pixels
+    if ( GetDerivedReservedTextureSize() < 5 )
+    {
+        SetIfValid(OutInvalidInfo, LOCTEXT("UseReservedTx.NoReservedSpace.ToolTip", "No space reserved in parent LOD chain"));
+        return false;
+    }
+
+    UEnum* EnumClass = StaticEnum<EHairTextureLayout>();
+    if ( !EnumClass || GetDerivedTextureChannelLayout() != EnumClass->GetNameStringByValue((int64)GetChannelLayout()) )
+    {
+        SetIfValid(OutInvalidInfo, LOCTEXT("UseReservedTx.InconsistentReservedLayout.ToolTip", "Parent reserved texture layout setting differs from current LOD texture layout"));
+        return false;
+    }
+    
+    return true;
+}
+
 /* UHairCardGeneratorGroupSettings
  *****************************************************************************/
 
@@ -1314,4 +1408,6 @@ void UHairCardGeneratorGroupSettings::SetStrandStats()
 	}
 
 	StrandCount = CountFiltered;
-}
\ No newline at end of file
+}
+
+#undef LOCTEXT_NAMESPACE
diff --git a/Engine/Plugins/Experimental/HairCardGenerator/Source/HairCardGeneratorEditor/Private/HairCardGeneratorPluginSettings.h b/Engine/Plugins/Experimental/HairCardGenerator/Source/HairCardGeneratorEditor/Private/HairCardGeneratorPluginSettings.h
--- a/Engine/Plugins/Experimental/HairCardGenerator/Source/HairCardGeneratorEditor/Private/HairCardGeneratorPluginSettings.h
+++ b/Engine/Plugins/Experimental/HairCardGenerator/Source/HairCardGeneratorEditor/Private/HairCardGeneratorPluginSettings.h
@@ -238,7 +238,10 @@ public:
 	int32 GetGenerateForGroomGroup() const { return GenerateForGroomGroup; }
 	void SetGenerateForGroomGroup(int GroupIndex) { GenerateForGroomGroup = GroupIndex; }
 	EHairTextureLayout GetChannelLayout() const { return ChannelLayout; }
-	bool ValidChannelLayouts() { return bValidChannelLayouts; }
+	bool ValidChannelLayouts() const { return bValidChannelLayouts; }
+
+	bool CanReduceFromLOD(FText* OutInvalidInfo = nullptr) const;
+	bool CanUseReservedTx(FText* OutInvalidInfo = nullptr) const;
 
 	TArray<TObjectPtr<UHairCardGeneratorGroupSettings>>& GetFilterGroupSettings() { return FilterGroupGenerationSettings; }
 	const TArray<TObjectPtr<UHairCardGeneratorGroupSettings>>& GetFilterGroupSettings() const { return FilterGroupGenerationSettings; }
@@ -255,6 +258,9 @@ private:
 
 	void UpdateParentInfo();
 	void UpdateHairWidths();
+
+	void EnforceValidLODSettings();
+
 	bool FindDerivedTextureSettings();
 	TSharedPtr<FJsonObject> GetParentTextureSettings() const;
 
@@ -282,7 +288,7 @@ private:
 
 	// Static HCS version number to force regeneration on tool updates
 	UPROPERTY(BlueprintReadOnly, Category="Asset", meta=(AllowPrivateAccess))
-	FString Version = TEXT("0.3");
+	FString Version = TEXT("0.4");
 
 	UPROPERTY(BlueprintReadOnly, Category="Settings Groups", meta=(AllowPrivateAccess))
 	TArray<TObjectPtr<UHairCardGeneratorGroupSettings>> FilterGroupGenerationSettings;
