HairCardSolver: Add UI element for forcing regeneration of hair cards regardless of matching settings

diff --git a/Engine/Plugins/Experimental/HairCardGenerator/Source/HairCardGeneratorEditor/Private/HairCardGenOptionsWindow.cpp b/Engine/Plugins/Experimental/HairCardGenerator/Source/HairCardGeneratorEditor/Private/HairCardGenOptionsWindow.cpp
--- a/Engine/Plugins/Experimental/HairCardGenerator/Source/HairCardGeneratorEditor/Private/HairCardGenOptionsWindow.cpp
+++ b/Engine/Plugins/Experimental/HairCardGenerator/Source/HairCardGeneratorEditor/Private/HairCardGenOptionsWindow.cpp
@@ -144,6 +144,9 @@ private:
 
 	void OnEditGroupSettings(TObjectPtr<UHairCardGeneratorGroupSettings> GroupSettings);
 
+	ECheckBoxState IsForceRegen() const;
+	void OnForceRegenToggle(ECheckBoxState NewState);
+
 	ESelectionMode::Type GetListSelectionMode() const { return ESelectionMode::Single; }
 	TSharedRef<ITableRow> OnGenerateOptionListRow(TObjectPtr<UHairCardGeneratorGroupSettings> InGroupSettings, const TSharedRef<STableViewBase>& OwnerTable);
 
@@ -238,6 +241,9 @@ bool HairCardGenWindow_Utils::PromptUserWithHairCardGenDialog(UHairCardGenerator
 		ParentWindow = MainFrame.GetParentWindow();
 	}
 
+	// Always set force regen to false on window open
+	// SettingsObject->SetForceRegenerate(false);
+
 	// Compute centered window position based on max window size, which include when all categories are expanded
 	const float WindowWidth = 600.0f;
 	const float WindowHeight = 400.0f;
@@ -568,6 +574,17 @@ void SHairCardGenOptionsWindow::Construct(const FArguments& InArgs)
 						.IsEnabled(this, &SHairCardGenOptionsWindow::CanOpenAdvancedSettings)
 				]
 				+ SUniformGridPanel::Slot(2, 0)
+				[
+					SNew(SCheckBox)
+					.HAlign(HAlign_Right)
+					.ToolTipText(LOCTEXT("HairCardSettings.ForceRegen.ToolTip", "Force full regeneration"))
+					.IsChecked(this, &SHairCardGenOptionsWindow::IsForceRegen)
+					.OnCheckStateChanged(this,  &SHairCardGenOptionsWindow::OnForceRegenToggle)
+					[
+						SNew(STextBlock).Text(LOCTEXT("HairCardSettings.ForceRegen", "Force Regenerate"))
+					]
+				]
+				+ SUniformGridPanel::Slot(3, 0)
 				[
 					SAssignNew(RunButton, SPrimaryButton)
 						.Text(LOCTEXT("HairCardSettings.Generate", "Generate"))
@@ -575,7 +592,7 @@ void SHairCardGenOptionsWindow::Construct(const FArguments& InArgs)
 						.OnClicked(this, &SHairCardGenOptionsWindow::OnGenerateClick)
 						.IsEnabled(this, &SHairCardGenOptionsWindow::CanRunGeneration)
 				]
-				+ SUniformGridPanel::Slot(3, 0)
+				+ SUniformGridPanel::Slot(4, 0)
 				[
 					SNew(SButton)
 						.HAlign(HAlign_Center)
@@ -620,6 +637,7 @@ void SHairCardGenOptionsWindow::Construct(const FArguments& InArgs)
 			]
 			+ SHorizontalBox::Slot()
 				.HAlign(HAlign_Right)
+				.VAlign(VAlign_Center)
 				.Padding(FMargin(2.0f, 0.0f))
 			[
 				SNew(SButton)
@@ -730,6 +748,16 @@ void SHairCardGenOptionsWindow::OnEditGroupSettings(TObjectPtr<UHairCardGenerato
 	HairCardGenWindow_Impl::PromptGroupSettingsDialog(WidgetWindow.Pin(), SettingsObject, FilterGroupIndex);
 }
 
+ECheckBoxState SHairCardGenOptionsWindow::IsForceRegen() const
+{
+	return (SettingsObject->GetForceRegenerate()) ? ECheckBoxState::Checked : ECheckBoxState::Unchecked;
+}
+
+void SHairCardGenOptionsWindow::OnForceRegenToggle(ECheckBoxState NewState)
+{
+	SettingsObject->SetForceRegenerate(NewState == ECheckBoxState::Checked);
+}
+
 FText SHairCardGenOptionsWindow::GetTargetAssetName() const
 {
 	return FText::FromString(SettingsObject->BaseFilename);
diff --git a/Engine/Plugins/Experimental/HairCardGenerator/Source/HairCardGeneratorEditor/Private/HairCardGeneratorPluginSettings.cpp b/Engine/Plugins/Experimental/HairCardGenerator/Source/HairCardGeneratorEditor/Private/HairCardGeneratorPluginSettings.cpp
--- a/Engine/Plugins/Experimental/HairCardGenerator/Source/HairCardGeneratorEditor/Private/HairCardGeneratorPluginSettings.cpp
+++ b/Engine/Plugins/Experimental/HairCardGenerator/Source/HairCardGeneratorEditor/Private/HairCardGeneratorPluginSettings.cpp
@@ -459,13 +459,22 @@ uint8 UHairCardGeneratorPluginSettings::GetAllPipelineGeneratedDifferences() con
 
 uint8 UHairCardGeneratorPluginSettings::GetPipelineGeneratedDifferences() const
 {
+	if ( bForceRegen )
+	{
+		return (uint8)EHairCardGenerationPipeline::All;
+	}
+
 	const FString SettingsBinFile = GetGeneratedSettingsFilename();
 	if ( !FPaths::FileExists(SettingsBinFile) )
+	{
 		return (uint8)EHairCardGenerationPipeline::All;
+	}
 
 	TArray<uint8> SettingsBytes;
 	if ( !FFileHelper::LoadFileToArray(SettingsBytes, *SettingsBinFile, 0) )
+	{
 		return (uint8)EHairCardGenerationPipeline::All;
+	}
 
 	FMemoryReader MemReader(SettingsBytes, true);
 
@@ -475,7 +484,9 @@ uint8 UHairCardGeneratorPluginSettings::GetPipelineGeneratedDifferences() const
 
 	// Early-out if HCS tool versions differ from generated settings
 	if (DiffObject->Version != Version)
+	{
 		return (uint8)EHairCardGenerationSettingsCategories::All;
+	}
 
 	return HairCardSettings_Helpers::PipelineSettingsDiff(this, DiffObject);
 }
diff --git a/Engine/Plugins/Experimental/HairCardGenerator/Source/HairCardGeneratorEditor/Private/HairCardGeneratorPluginSettings.h b/Engine/Plugins/Experimental/HairCardGenerator/Source/HairCardGeneratorEditor/Private/HairCardGeneratorPluginSettings.h
--- a/Engine/Plugins/Experimental/HairCardGenerator/Source/HairCardGeneratorEditor/Private/HairCardGeneratorPluginSettings.h
+++ b/Engine/Plugins/Experimental/HairCardGenerator/Source/HairCardGeneratorEditor/Private/HairCardGeneratorPluginSettings.h
@@ -243,6 +243,9 @@ public:
 	bool CanReduceFromLOD(FText* OutInvalidInfo = nullptr) const;
 	bool CanUseReservedTx(FText* OutInvalidInfo = nullptr) const;
 
+	bool GetForceRegenerate() const {return bForceRegen;}
+	void SetForceRegenerate(bool bForceRegenerate){bForceRegen = bForceRegenerate;}
+
 	TArray<TObjectPtr<UHairCardGeneratorGroupSettings>>& GetFilterGroupSettings() { return FilterGroupGenerationSettings; }
 	const TArray<TObjectPtr<UHairCardGeneratorGroupSettings>>& GetFilterGroupSettings() const { return FilterGroupGenerationSettings; }
 
@@ -324,6 +327,7 @@ private:
 	UPROPERTY(Transient)
 	TObjectPtr<const UStaticMesh> OldGeneratedMesh = nullptr;
 
+	bool bForceRegen;
 	uint32 UnassignedStrandsCount;
 	TMap<uint32,uint32> StrandErrorCountMap;
 };
