HairCardSolver: Fixed dimension issue with groom groups when width overrides are set. Also, increased max card count to 100k, and don't run DBSCAN when flyaways set to 0.

diff --git a/Engine/Plugins/Experimental/HairCardGenerator/Source/HairCardGeneratorEditor/Private/HairCardGeneratorPluginSettings.h b/Engine/Plugins/Experimental/HairCardGenerator/Source/HairCardGeneratorEditor/Private/HairCardGeneratorPluginSettings.h
--- a/Engine/Plugins/Experimental/HairCardGenerator/Source/HairCardGeneratorEditor/Private/HairCardGeneratorPluginSettings.h
+++ b/Engine/Plugins/Experimental/HairCardGenerator/Source/HairCardGeneratorEditor/Private/HairCardGeneratorPluginSettings.h
@@ -367,7 +367,7 @@ public:
 	TSet<FName> ApplyToCardGroups = {NAME_None};
 
 	// Total number of cards to generate for this settings group
-	UPROPERTY(EditAnywhere, BlueprintReadOnly, Category="Cards", meta=(ClampMin="1", ClampMax="10000"))
+	UPROPERTY(EditAnywhere, BlueprintReadOnly, Category="Cards", meta=(ClampMin="1", ClampMax="100000"))
 	int32 TargetNumberOfCards = 2000;
 
 	// Target number of triangles of the final mesh, only for adaptive subdivision
