[ToolWidgets] Dependencies reference correct base widget

diff --git a/Engine/Plugins/Experimental/HairCardGenerator/Source/HairCardGeneratorEditor/Private/HairCardGenOptionsWindow.cpp b/Engine/Plugins/Experimental/HairCardGenerator/Source/HairCardGeneratorEditor/Private/HairCardGenOptionsWindow.cpp
--- a/Engine/Plugins/Experimental/HairCardGenerator/Source/HairCardGeneratorEditor/Private/HairCardGenOptionsWindow.cpp
+++ b/Engine/Plugins/Experimental/HairCardGenerator/Source/HairCardGeneratorEditor/Private/HairCardGenOptionsWindow.cpp
@@ -162,7 +162,7 @@ private:
 	TSharedPtr<class IDetailsView> DetailsView;
 	// List view for updating per settings-group options
 	TSharedPtr<SListView<TObjectPtr<UHairCardGeneratorGroupSettings>>> GroupSettingsListView;
-	TSharedPtr<SButton> RunButton;
+	TSharedPtr<SWidget> RunButton;
 };
 
 
@@ -213,7 +213,7 @@ private:
 	TWeakPtr<SWindow> WidgetWindow;
 	// Group settings detail view
 	TSharedPtr<IDetailsView> DetailsView;
-	TSharedPtr<SButton> OkButton;
+	TSharedPtr<SWidget> OkButton;
 };
 
 
diff --git a/Engine/Plugins/Runtime/ModelViewViewModel/Source/ModelViewViewModelEditor/Private/Widgets/SMVVMFieldSelectorMenu.h b/Engine/Plugins/Runtime/ModelViewViewModel/Source/ModelViewViewModelEditor/Private/Widgets/SMVVMFieldSelectorMenu.h
--- a/Engine/Plugins/Runtime/ModelViewViewModel/Source/ModelViewViewModelEditor/Private/Widgets/SMVVMFieldSelectorMenu.h
+++ b/Engine/Plugins/Runtime/ModelViewViewModel/Source/ModelViewViewModelEditor/Private/Widgets/SMVVMFieldSelectorMenu.h
@@ -180,7 +180,7 @@ private:
 	FFieldSelectionContext SelectionContext;
 
 	TSharedPtr<SSearchBox> SearchBox;
-	TSharedPtr<SPrimaryButton> EventSelectButton;
+	TSharedPtr<SWidget> EventSelectButton;
 
 	//~ viewmodels (binding context panel)
 	TSharedPtr<SListView<FBindingSource>> ViewModelList;
diff --git a/Engine/Source/Editor/Documentation/Private/SDocumentationAnchor.h b/Engine/Source/Editor/Documentation/Private/SDocumentationAnchor.h
--- a/Engine/Source/Editor/Documentation/Private/SDocumentationAnchor.h
+++ b/Engine/Source/Editor/Documentation/Private/SDocumentationAnchor.h
@@ -32,6 +32,6 @@ private:
 
 	TAttribute<FString> Link;
 	TAttribute<FString> BaseUrlId;
-	TSharedPtr<class SButton> Button;
+	TSharedPtr<class SWidget> Button;
 	TSharedPtr<class SImage> ButtonImage;
 };
diff --git a/Engine/Source/Editor/NewLevelDialog/Private/NewLevelDialogModule.cpp b/Engine/Source/Editor/NewLevelDialog/Private/NewLevelDialogModule.cpp
--- a/Engine/Source/Editor/NewLevelDialog/Private/NewLevelDialogModule.cpp
+++ b/Engine/Source/Editor/NewLevelDialog/Private/NewLevelDialogModule.cpp
@@ -244,7 +244,7 @@ public:
 			.ItemHeight(NewLevelDialogDefs::TemplateTileHeight + 9)
 			.ItemWidth(NewLevelDialogDefs::TemplateTileWidth + 9);
 
-		TSharedPtr<SButton> CreateButton;
+		TSharedPtr<SWidget> CreateButton;
 
 		this->ChildSlot
 		[
diff --git a/Engine/Source/Editor/SourceControlWindows/Private/SSourceControlCheckedOutDialog.h b/Engine/Source/Editor/SourceControlWindows/Private/SSourceControlCheckedOutDialog.h
--- a/Engine/Source/Editor/SourceControlWindows/Private/SSourceControlCheckedOutDialog.h
+++ b/Engine/Source/Editor/SourceControlWindows/Private/SSourceControlCheckedOutDialog.h
@@ -107,7 +107,7 @@ private:
 	EColumnSortMode::Type SortMode;
 
 	/** The close button widget. */
-	TSharedPtr<SButton> CloseButton;
+	TSharedPtr<SWidget> CloseButton;
 
 	/** The checkbox widget. */
 	TSharedPtr<SCheckBox> CheckBox;
diff --git a/Engine/Source/Editor/SparseVolumeTexture/Private/OpenVDBImportWindow.h b/Engine/Source/Editor/SparseVolumeTexture/Private/OpenVDBImportWindow.h
--- a/Engine/Source/Editor/SparseVolumeTexture/Private/OpenVDBImportWindow.h
+++ b/Engine/Source/Editor/SparseVolumeTexture/Private/OpenVDBImportWindow.h
@@ -105,7 +105,7 @@ private:
 	TSharedPtr<SOpenVDBAttributesConfigurator> AttributesAConfigurator;
 	TSharedPtr<SOpenVDBAttributesConfigurator> AttributesBConfigurator;
 	TSharedPtr<SCheckBox> ImportAsSequenceCheckBox;
-	TSharedPtr<SButton> ImportButton;
+	TSharedPtr<SWidget> ImportButton;
 	TWeakPtr<SWindow> WidgetWindow;
 	bool bShouldImport;
 
diff --git a/Engine/Source/Editor/UnrealEd/Private/FbxOptionWindow.h b/Engine/Source/Editor/UnrealEd/Private/FbxOptionWindow.h
--- a/Engine/Source/Editor/UnrealEd/Private/FbxOptionWindow.h
+++ b/Engine/Source/Editor/UnrealEd/Private/FbxOptionWindow.h
@@ -107,7 +107,7 @@ private:
 	UFbxImportUI*	ImportUI;
 	TSharedPtr<class IDetailsView> DetailsView;
 	TWeakPtr< SWindow > WidgetWindow;
-	TSharedPtr<SButton> ImportAllButton;
+	TSharedPtr<SWidget> ImportAllButton;
 	bool			bShouldImport;
 	bool			bShouldImportAll;
 	bool			bIsObjFormat;
