8-Hair card generator dataflow

diff --git a/Engine/Plugins/Experimental/HairCardGenerator/HairCardGenerator.uplugin b/Engine/Plugins/Experimental/HairCardGenerator/HairCardGenerator.uplugin
--- a/Engine/Plugins/Experimental/HairCardGenerator/HairCardGenerator.uplugin
+++ b/Engine/Plugins/Experimental/HairCardGenerator/HairCardGenerator.uplugin
@@ -24,6 +24,15 @@
 				"Win64",
 				"Linux"
 			]
+		},
+		{
+			"Name": "HairCardGeneratorDataflow",
+			"Type": "Editor",
+			"LoadingPhase": "Default",
+			"PlatformAllowList": [
+				"Win64",
+				"Linux"
+			]
 		}
 	],
 	"Plugins": [
@@ -34,6 +43,10 @@
 		{
 			"Name": "PythonFoundationPackages",
 			"Enabled": true
+		},
+		{
+			"Name": "Dataflow",
+			"Enabled": true
 		}
 	],
 	"PythonRequirements":
diff --git a/Engine/Plugins/Experimental/HairCardGenerator/Source/HairCardGeneratorDataflow/HairCardGeneratorDataflow.Build.cs b/Engine/Plugins/Experimental/HairCardGenerator/Source/HairCardGeneratorDataflow/HairCardGeneratorDataflow.Build.cs
new file mode 100644
--- /dev/null
+++ b/Engine/Plugins/Experimental/HairCardGenerator/Source/HairCardGeneratorDataflow/HairCardGeneratorDataflow.Build.cs
@@ -0,0 +1,33 @@
+// Copyright Epic Games, Inc. All Rights Reserved.
+
+namespace UnrealBuildTool.Rules
+{
+	public class HairCardGeneratorDataflow: ModuleRules
+	{
+		public HairCardGeneratorDataflow(ReadOnlyTargetRules Target) : base(Target)
+		{
+			PrivateIncludePaths.AddRange(new string[] {
+				System.IO.Path.Combine(GetModuleDirectory("HairCardGeneratorEditor"), "Private"),
+			});
+			
+			PublicDependencyModuleNames.AddRange(
+				new string[]
+				{
+					"Core",
+					"CoreUObject",
+					"ChaosCore",
+					"Chaos",
+					"DataflowCore",
+					"DataflowEngine",
+					"DataflowEditor",
+					"DataflowNodes",
+					"HairStrandsCore",
+					"HairStrandsDataflow",
+					"HairCardGeneratorEditor",
+					"HairCardGeneratorFramework",
+					"Engine",
+				}
+			);
+		}
+	}
+}
\ No newline at end of file
diff --git a/Engine/Plugins/Experimental/HairCardGenerator/Source/HairCardGeneratorDataflow/Private/BuildCardsSettingsNode.cpp b/Engine/Plugins/Experimental/HairCardGenerator/Source/HairCardGeneratorDataflow/Private/BuildCardsSettingsNode.cpp
new file mode 100644
--- /dev/null
+++ b/Engine/Plugins/Experimental/HairCardGenerator/Source/HairCardGeneratorDataflow/Private/BuildCardsSettingsNode.cpp
@@ -0,0 +1,74 @@
+// Copyright Epic Games, Inc. All Rights Reserved.
+
+#include "BuildCardsSettingsNode.h"
+
+#include "GroomCollectionFacades.h"
+#include "HairCardGeneratorEditorModule.h"
+
+#include UE_INLINE_GENERATED_CPP_BY_NAME(BuildCardsSettingsNode)
+
+namespace UE::CardGen::Private
+{
+	static void BuildStrandsPositions(const FManagedArrayCollection& GroomCollection, TArray<TArray<FVector>>& StrandsPositions)
+	{
+		const UE::Groom::FGroomStrandsFacade StrandsFacade(GroomCollection);
+		if (StrandsFacade.IsValid())
+		{
+			const int32 NumStrands = StrandsFacade.GetNumCurves();
+			StrandsPositions.SetNum(NumStrands);
+
+			int32 PointOffset = 0;
+			for(int32 StrandIndex = 0; StrandIndex < NumStrands; ++StrandIndex)
+			{
+				const int32 NumPoints = StrandsFacade.GetCurvePointOffsets()[StrandIndex]-PointOffset;
+				StrandsPositions[StrandIndex].SetNum(NumPoints);
+				
+				for(int32 PointIndex = 0; PointIndex < NumPoints; ++PointIndex)
+				{
+					StrandsPositions[StrandIndex][PointIndex] = FVector(StrandsFacade.GetPointRestPositions()[PointIndex + PointOffset]);
+				}
+				PointOffset = StrandsFacade.GetCurvePointOffsets()[StrandIndex];
+			}
+		}
+	}
+}
+void FBuildCardsSettingsNode::Evaluate(UE::Dataflow::FContext& Context, const FDataflowOutput* Out) const
+{
+	if (Out->IsA<FManagedArrayCollection>(&Collection))
+	{
+		FManagedArrayCollection GroomCollection = GetValue<FManagedArrayCollection>(Context, &Collection);
+		SetValue(Context, MoveTemp(GroomCollection), &Collection);
+	}
+	else if (Out->IsA<TArray<FGroomCardsSettings>>(&CardsSettings))
+	{
+		UGroomAsset* LocalAsset = const_cast<UGroomAsset*>(GroomAsset.Get());
+		if (!LocalAsset)
+		{
+			if (const UE::Dataflow::FEngineContext* EngineContext = Context.AsType<UE::Dataflow::FEngineContext>())
+			{
+				LocalAsset = Cast<UGroomAsset>(EngineContext->Owner);
+			}
+		}
+		TArray<FGroomCardsSettings> OutputSettings;
+		if(LocalAsset && !LocalAsset->GetHairGroupsCards().IsEmpty())
+		{
+			FManagedArrayCollection GroomCollection = GetValue<FManagedArrayCollection>(Context, &Collection);
+			for(FHairGroupsCardsSourceDescription& HairCardsDescription : LocalAsset->GetHairGroupsCards())
+			{
+				HairCardsDescription.Textures = FHairGroupCardsTextures();
+				
+				OutputSettings.AddDefaulted();
+				FHairCardGeneratorUtils::BuildGenerationSettings(false, LocalAsset, HairCardsDescription,
+					OutputSettings.Last().GenerationSettings, OutputSettings.Last().GenerationFlags, OutputSettings.Last().PipelineFlags);
+
+				FHairCardGeneratorUtils::LoadGroomStrands(LocalAsset, [&GroomCollection](TArray<TArray<FVector>>& StrandsPositions)
+				{
+					UE::CardGen::Private::BuildStrandsPositions(GroomCollection, StrandsPositions);
+				});
+				OutputSettings.Last().GenerationSettings->UpdateStrandFilterAssignment();
+			}
+		}
+		SetValue(Context, MoveTemp(OutputSettings), &CardsSettings);
+	}
+}
+
diff --git a/Engine/Plugins/Experimental/HairCardGenerator/Source/HairCardGeneratorDataflow/Private/CardsAssetTerminalNode.cpp b/Engine/Plugins/Experimental/HairCardGenerator/Source/HairCardGeneratorDataflow/Private/CardsAssetTerminalNode.cpp
new file mode 100644
--- /dev/null
+++ b/Engine/Plugins/Experimental/HairCardGenerator/Source/HairCardGeneratorDataflow/Private/CardsAssetTerminalNode.cpp
@@ -0,0 +1,46 @@
+// Copyright Epic Games, Inc. All Rights Reserved.
+
+#include "CardsAssetTerminalNode.h"
+
+#include "BuildCardsSettingsNode.h"
+#include "GroomCollectionFacades.h"
+#include "GroomEdit.h"
+#include "HairCardGenControllerBase.h"
+#include "HairCardGeneratorEditorModule.h"
+#include "HairCardGeneratorEditorSettings.h"
+#include "Dataflow/DataflowObjectInterface.h"
+
+#include UE_INLINE_GENERATED_CPP_BY_NAME(CardsAssetTerminalNode)
+
+namespace UE::Groom::Private
+{
+}
+
+void FCardsAssetTerminalNode::SetAssetValue(TObjectPtr<UObject> Asset, UE::Dataflow::FContext& Context) const
+{
+	if (UGroomAsset* GroomAsset = Cast<UGroomAsset>(Asset.Get()))
+	{
+		TArray<FGroomCardsSettings> LocalSettings = GetValue<TArray<FGroomCardsSettings>>(Context, &CardsSettings);
+
+		if(GroomAsset && !GroomAsset->GetHairGroupsCards().IsEmpty() && (LocalSettings.Num() == GroomAsset->GetHairGroupsCards().Num()))
+		{
+			GroomAsset->Modify();
+			int32 GroupIndex = 0;
+			for(FHairGroupsCardsSourceDescription& HairCardsDescription : GroomAsset->GetHairGroupsCards())
+			{
+				FHairCardGeneratorUtils::BuildCardsAssets(GroomAsset, HairCardsDescription,
+					LocalSettings[GroupIndex].GenerationSettings,  LocalSettings[GroupIndex].GenerationFlags);
+				++GroupIndex;
+			}
+		}
+	}
+}
+
+void FCardsAssetTerminalNode::Evaluate(UE::Dataflow::FContext& Context, const FDataflowOutput* Out) const
+{
+	const FManagedArrayCollection& InCollection = GetValue<FManagedArrayCollection>(Context, &Collection);
+	SetValue(Context, InCollection, &Collection);
+}
+
+
+
diff --git a/Engine/Plugins/Experimental/HairCardGenerator/Source/HairCardGeneratorDataflow/Private/GenerateCardsClumpsNode.cpp b/Engine/Plugins/Experimental/HairCardGenerator/Source/HairCardGeneratorDataflow/Private/GenerateCardsClumpsNode.cpp
new file mode 100644
--- /dev/null
+++ b/Engine/Plugins/Experimental/HairCardGenerator/Source/HairCardGeneratorDataflow/Private/GenerateCardsClumpsNode.cpp
@@ -0,0 +1,97 @@
+// Copyright Epic Games, Inc. All Rights Reserved.
+
+#include "GenerateCardsClumpsNode.h"
+#include "GroomCollectionFacades.h"
+#include "HairCardGeneratorEditorModule.h"
+#include "Dataflow/DataflowDebugDrawInterface.h"
+
+#include UE_INLINE_GENERATED_CPP_BY_NAME(GenerateCardsClumpsNode)
+
+// Clumps Attributes
+const FName FGenerateCardsClumpsNode::CurveClumpIndicesAttribute("CurveClumpIndices_LOD");
+const FName FGenerateCardsClumpsNode::ObjectNumClumpsAttribute("ObjectNumClumps_LOD");
+
+void FGenerateCardsClumpsNode::Evaluate(UE::Dataflow::FContext& Context, const FDataflowOutput* Out) const
+{
+	if (Out->IsA<FManagedArrayCollection>(&Collection) || Out->IsA<TArray<FGroomCardsSettings>>(&CardsSettings))
+	{
+		TArray<FGroomCardsSettings> OutputSettings = GetValue<TArray<FGroomCardsSettings>>(Context, &CardsSettings);
+		for(FGroomCardsSettings& LODSettings : OutputSettings)
+		{
+			// Override the generation settings if matching the LOD index and card group
+			if(TObjectPtr<UHairCardGeneratorPluginSettings>
+				GenerationSettings = LODSettings.GenerationSettings)
+			{
+				for(const FCardsClumpsSettings& OverideSettings : ClumpsSettings)
+				{
+					if(GenerationSettings->GetLODIndex() == OverideSettings.LODIndex || OverideSettings.LODIndex == INDEX_NONE)
+					{
+						for(TObjectPtr<UHairCardGeneratorGroupSettings> FilterSettings : GenerationSettings->GetFilterGroupSettings())
+						{
+							if(FilterSettings.Get())
+							{
+								if((OverideSettings.CardGroup != NAME_None && FilterSettings->ApplyToCardGroups.Contains(OverideSettings.CardGroup)) || OverideSettings.CardGroup == NAME_None)
+								{
+									FilterSettings->TargetNumberOfCards = OverideSettings.NumCards;
+									FilterSettings->MaxNumberOfFlyaways = OverideSettings.NumFlyaways;
+								}
+							}
+						}
+					}
+				}
+			}
+		}
+		if(Out->IsA<TArray<FGroomCardsSettings>>(&CardsSettings))
+		{
+			SetValue(Context, MoveTemp(OutputSettings), &CardsSettings);
+		}
+		else if(Out->IsA<FManagedArrayCollection>(&Collection))
+		{
+			FManagedArrayCollection GroomCollection = GetValue<FManagedArrayCollection>(Context, &Collection);
+			
+			for(FGroomCardsSettings& LODSettings : OutputSettings)
+			{
+				if(TObjectPtr<UHairCardGeneratorPluginSettings> GenerationSettings = LODSettings.GenerationSettings)
+				{
+					if(FHairCardGeneratorUtils::LoadGenerationSettings(GenerationSettings))
+					{
+						TArray<int32> StrandsClumps;
+						int32 NumClumps;
+						
+						bool bHasClumps = FHairCardGeneratorUtils::RunCardsGeneration(GenerationSettings, LODSettings.PipelineFlags,
+				[&StrandsClumps, &NumClumps](const TObjectPtr<const UHairCardGeneratorPluginSettings>& GenerationSettings, const int32 FilterIndex, const uint8 GenFlags)
+						{
+							return FHairCardGeneratorUtils::GenerateCardsClumps(GenerationSettings, FilterIndex, GenFlags, StrandsClumps, NumClumps);
+						}, false);
+			
+						if(bHasClumps)
+						{
+							FString ClumpIndicesLOD = CurveClumpIndicesAttribute.ToString();
+							ClumpIndicesLOD.AppendInt(GenerationSettings->GetLODIndex());
+			
+							FString NumClumpsLOD = ObjectNumClumpsAttribute.ToString();
+							NumClumpsLOD.AppendInt(GenerationSettings->GetLODIndex());
+					
+							TManagedArray<int32>& CurveClumpIndices = GroomCollection.AddAttribute<int32>(FName(ClumpIndicesLOD), UE::Groom::FGroomStrandsFacade::CurvesGroup);
+							TManagedArray<int32>& ObjectNumClumps = GroomCollection.AddAttribute<int32>(FName(NumClumpsLOD), UE::Groom::FGroomStrandsFacade::ObjectsGroup);
+							
+							if(CurveClumpIndices.Num() == (StrandsClumps.Num()-1))
+							{
+								for(int32 CurveIndex = 0, NumCurves = CurveClumpIndices.Num(); CurveIndex < NumCurves; ++CurveIndex)
+								{
+									CurveClumpIndices[CurveIndex] = StrandsClumps[CurveIndex];
+								}
+							}
+							for(int32 ObjectIndex = 0, NumObjects = ObjectNumClumps.Num(); ObjectIndex < NumObjects; ++ObjectIndex)
+							{
+								ObjectNumClumps[ObjectIndex] = NumClumps;
+							}
+						}
+					}
+				}
+			}
+			SetValue(Context, MoveTemp(GroomCollection), &Collection);
+		}
+	}
+}
+
diff --git a/Engine/Plugins/Experimental/HairCardGenerator/Source/HairCardGeneratorDataflow/Private/GenerateCardsGeometryNode.cpp b/Engine/Plugins/Experimental/HairCardGenerator/Source/HairCardGeneratorDataflow/Private/GenerateCardsGeometryNode.cpp
new file mode 100644
--- /dev/null
+++ b/Engine/Plugins/Experimental/HairCardGenerator/Source/HairCardGeneratorDataflow/Private/GenerateCardsGeometryNode.cpp
@@ -0,0 +1,126 @@
+// Copyright Epic Games, Inc. All Rights Reserved.
+
+#include "GenerateCardsGeometryNode.h"
+#include "GenerateCardsClumpsNode.h"
+#include "GroomCollectionFacades.h"
+#include "HairCardGeneratorEditorModule.h"
+
+#include UE_INLINE_GENERATED_CPP_BY_NAME(GenerateCardsGeometryNode)
+
+// Geometry Attributes
+const FName FGenerateCardsGeometryNode::VertexClumpPositionsAttribute("VertexClumpPositions");
+const FName FGenerateCardsGeometryNode::FaceVertexIndicesAttribute("FaceVertexIndices");
+const FName FGenerateCardsGeometryNode::VertexCardIndicesAttribute("VertexCardIndices");
+
+const FName FGenerateCardsGeometryNode::CardsVerticesGroup("CardsVertices_LOD");
+const FName FGenerateCardsGeometryNode::CardsFacesGroup("CardsFaces_LOD");
+
+void FGenerateCardsGeometryNode::Evaluate(UE::Dataflow::FContext& Context, const FDataflowOutput* Out) const
+{
+	if (Out->IsA<FManagedArrayCollection>(&Collection) || Out->IsA<TArray<FGroomCardsSettings>>(&CardsSettings))
+	{
+		TArray<FGroomCardsSettings> OutputSettings = GetValue<TArray<FGroomCardsSettings>>(Context, &CardsSettings);
+		for(FGroomCardsSettings& LODSettings : OutputSettings)
+		{
+			// Override the generation settings if matching the LOD index and card group
+			if(TObjectPtr<UHairCardGeneratorPluginSettings> GenerationSettings = LODSettings.GenerationSettings)
+			{
+				for(const FCardsGeometrySettings& OverideSettings : GeometrySettings)
+				{
+					if(GenerationSettings->GetLODIndex() == OverideSettings.LODIndex || OverideSettings.LODIndex == INDEX_NONE)
+					{
+						for(TObjectPtr<UHairCardGeneratorGroupSettings> FilterSettings : GenerationSettings->GetFilterGroupSettings())
+						{
+							if(FilterSettings.Get())
+                        	{
+								if((OverideSettings.CardGroup != NAME_None && FilterSettings->ApplyToCardGroups.Contains(OverideSettings.CardGroup)) || OverideSettings.CardGroup == NAME_None)
+								{
+									FilterSettings->TargetTriangleCount = OverideSettings.NumTriangles;
+								}
+							}
+						}
+					}
+				}
+			}
+		}
+		if(Out->IsA<TArray<FGroomCardsSettings>>(&CardsSettings))
+		{
+			SetValue(Context, MoveTemp(OutputSettings), &CardsSettings);
+		}
+		else if(Out->IsA<FManagedArrayCollection>(&Collection))
+		{
+			FManagedArrayCollection GroomCollection = GetValue<FManagedArrayCollection>(Context, &Collection);
+			
+			for(FGroomCardsSettings& LODSettings : OutputSettings)
+			{
+				if(TObjectPtr<UHairCardGeneratorPluginSettings> GenerationSettings = LODSettings.GenerationSettings)
+				{
+					if(FHairCardGeneratorUtils::LoadGenerationSettings(GenerationSettings))
+					{
+						TArray<TArray<FVector3f>> ClumpsGeometry;
+						
+						bool bHasGeometry = FHairCardGeneratorUtils::RunCardsGeneration(GenerationSettings, LODSettings.PipelineFlags,
+                        	[&ClumpsGeometry](const TObjectPtr<const UHairCardGeneratorPluginSettings>& GenerationSettings, const int32 FilterIndex, const uint8 GenFlags)
+						{
+							return FHairCardGeneratorUtils::GenerateCardsGeometry(GenerationSettings, FilterIndex, GenFlags, ClumpsGeometry);
+						}, false);
+   
+						if(bHasGeometry)
+						{
+							TArray<FIntVector3> ClumpsFaces;
+							TArray<FVector3f> ClumpsVertices;
+							TArray<int32> CardIndices;
+							int32 VertexOffset = 0;
+   
+							int32 CardIndex = 0;
+							for(TArray<FVector3f>& GeometryVertices : ClumpsGeometry)
+							{
+								for(int32 CurveVertex = 0, NumVertices = GeometryVertices.Num()/2; CurveVertex < (NumVertices-1); ++CurveVertex)
+								{
+									const int32 VertexIndex = CurveVertex * 2 + VertexOffset;
+									ClumpsFaces.Add(FIntVector3(VertexIndex, VertexIndex + 2, VertexIndex + 3));
+									ClumpsFaces.Add(FIntVector3(VertexIndex + 3, VertexIndex + 1, VertexIndex));
+								}
+								for(FVector3f& VertexPosition : GeometryVertices)
+                                {
+									ClumpsVertices.Add(VertexPosition);
+									CardIndices.Add(CardIndex);
+								}
+								VertexOffset = ClumpsVertices.Num();
+								++CardIndex;
+							}
+							
+							FString CardsVerticesLODGroup = CardsVerticesGroup.ToString();
+							CardsVerticesLODGroup.AppendInt(GenerationSettings->GetLODIndex());
+							
+							FString CardsFacesLODGroup = CardsFacesGroup.ToString();
+							CardsFacesLODGroup.AppendInt(GenerationSettings->GetLODIndex());
+							
+							TManagedArray<FVector3f>& VertexClumpPositions = GroomCollection.AddAttribute<FVector3f>(VertexClumpPositionsAttribute, FName(CardsVerticesLODGroup));
+							TManagedArray<FIntVector3>& FaceVertexIndices = GroomCollection.AddAttribute<FIntVector3>(FaceVertexIndicesAttribute, FName(CardsFacesLODGroup));
+							TManagedArray<int32>& VertexCardIndices = GroomCollection.AddAttribute<int32>(VertexCardIndicesAttribute, FName(CardsVerticesLODGroup));
+							
+							GroomCollection.EmptyGroup(FName(CardsVerticesLODGroup));
+							GroomCollection.AddElements(ClumpsVertices.Num(), FName(CardsVerticesLODGroup));
+							
+							GroomCollection.EmptyGroup(FName(CardsFacesLODGroup));
+                            GroomCollection.AddElements(ClumpsFaces.Num(), FName(CardsFacesLODGroup));
+							
+							for(int32 VertexIndex = 0, NumVertices = ClumpsVertices.Num(); VertexIndex < NumVertices; ++VertexIndex)
+							{
+								VertexClumpPositions[VertexIndex] = ClumpsVertices[VertexIndex];
+								VertexCardIndices[VertexIndex] = CardIndices[VertexIndex];
+							}
+							for(int32 FaceIndex = 0, NumFaces = ClumpsFaces.Num(); FaceIndex < NumFaces; ++FaceIndex)
+							{
+								FaceVertexIndices[FaceIndex] = ClumpsFaces[FaceIndex];
+							}
+						}
+					}
+				}
+			}
+			SetValue(Context, MoveTemp(GroomCollection), &Collection);
+		}
+	}
+}
+
diff --git a/Engine/Plugins/Experimental/HairCardGenerator/Source/HairCardGeneratorDataflow/Private/GenerateCardsTexturesNode.cpp b/Engine/Plugins/Experimental/HairCardGenerator/Source/HairCardGeneratorDataflow/Private/GenerateCardsTexturesNode.cpp
new file mode 100644
--- /dev/null
+++ b/Engine/Plugins/Experimental/HairCardGenerator/Source/HairCardGeneratorDataflow/Private/GenerateCardsTexturesNode.cpp
@@ -0,0 +1,118 @@
+// Copyright Epic Games, Inc. All Rights Reserved.
+
+#include "GenerateCardsTexturesNode.h"
+
+#include "GenerateCardsGeometryNode.h"
+#include "GroomCollectionFacades.h"
+#include "HairCardGeneratorEditorModule.h"
+
+#include UE_INLINE_GENERATED_CPP_BY_NAME(GenerateCardsTexturesNode)
+
+// Texture Attributes
+const FName FGenerateCardsTexturesNode::ObjectTextureIndicesAttribute("ObjectTextureIndices");
+const FName FGenerateCardsTexturesNode::CardsObjectsGroup("CardsObjects_LOD");
+const FName FGenerateCardsTexturesNode::VertexTextureUVsAttribute("VertexTextureUVs");
+
+void FGenerateCardsTexturesNode::Evaluate(UE::Dataflow::FContext& Context, const FDataflowOutput* Out) const
+{
+	if (Out->IsA<FManagedArrayCollection>(&Collection) || Out->IsA<TArray<FGroomCardsSettings>>(&CardsSettings))
+	{
+		TArray<FGroomCardsSettings> OutputSettings = GetValue<TArray<FGroomCardsSettings>>(Context, &CardsSettings);
+		for(FGroomCardsSettings& LODSettings : OutputSettings)
+		{
+			// Override the generation settings if matching the LOD index and card group
+			if(TObjectPtr<UHairCardGeneratorPluginSettings> GenerationSettings = LODSettings.GenerationSettings)
+			{
+				for(const FCardsTextureSettings& OverideSettings : TextureSettings)
+				{
+					if(GenerationSettings->GetLODIndex() == OverideSettings.LODIndex || OverideSettings.LODIndex == INDEX_NONE)
+					{
+						for(TObjectPtr<UHairCardGeneratorGroupSettings> FilterSettings : GenerationSettings->GetFilterGroupSettings())
+						{
+							if(FilterSettings.Get())
+							{
+								if((OverideSettings.CardGroup != NAME_None && FilterSettings->ApplyToCardGroups.Contains(OverideSettings.CardGroup)) || OverideSettings.CardGroup == NAME_None)
+								{
+									FilterSettings->NumberOfTexturesInAtlas = OverideSettings.NumTextures;
+								}
+							}
+						}
+					}
+				}
+			}
+		}
+		if(Out->IsA<TArray<FGroomCardsSettings>>(&CardsSettings))
+		{
+			SetValue(Context, MoveTemp(OutputSettings), &CardsSettings);
+		}
+		else if (Out->IsA<FManagedArrayCollection>(&Collection))
+		{
+			FManagedArrayCollection GroomCollection = GetValue<FManagedArrayCollection>(Context, &Collection);
+
+			for(FGroomCardsSettings& LODSettings : OutputSettings)
+			{
+				if(TObjectPtr<UHairCardGeneratorPluginSettings> GenerationSettings = LODSettings.GenerationSettings)
+				{
+					if(FHairCardGeneratorUtils::LoadGenerationSettings(GenerationSettings))
+					{
+						TArray<int32> CardsTextures;
+						int32 NumTextures = 0;
+						TArray<float> VertexUVs;
+						
+						bool bHasTextures = FHairCardGeneratorUtils::RunCardsGeneration(GenerationSettings, LODSettings.PipelineFlags,
+							[&CardsTextures, &NumTextures, &VertexUVs](const TObjectPtr<const UHairCardGeneratorPluginSettings>& GenerationSettings, const int32 FilterIndex, const uint8 GenFlags)
+						{
+							return FHairCardGeneratorUtils::GenerateCardsTextures(GenerationSettings, FilterIndex, GenFlags, CardsTextures, NumTextures, VertexUVs);
+						}, false);
+			
+						if(bHasTextures)
+						{
+							const int32 NumCards = CardsTextures.Num()-1-NumTextures;
+							
+							FString CardsObjectsLODGroup = CardsObjectsGroup.ToString();
+							CardsObjectsLODGroup.AppendInt(GenerationSettings->GetLODIndex());
+			
+							FString CardsVerticesLODGroup = FGenerateCardsGeometryNode::CardsVerticesGroup.ToString();
+							CardsVerticesLODGroup.AppendInt(GenerationSettings->GetLODIndex());
+							
+							TManagedArray<int32>& ObjectTextureIndices = GroomCollection.AddAttribute<int32>(ObjectTextureIndicesAttribute, FName(CardsObjectsLODGroup));
+							TManagedArray<FVector2f>& VertexTextureUvs = GroomCollection.AddAttribute<FVector2f>(VertexTextureUVsAttribute, FName(CardsVerticesLODGroup));
+							
+							GroomCollection.EmptyGroup(FName(CardsObjectsLODGroup));
+							GroomCollection.AddElements(NumCards, FName(CardsObjectsLODGroup));
+							
+							for(int32 CardIndex = 0; CardIndex < NumCards; ++CardIndex)
+							{
+								ObjectTextureIndices[CardIndex] = (CardsTextures[CardIndex] < NumTextures) ? CardsTextures[NumCards+CardsTextures[CardIndex]] : INDEX_NONE;
+							}
+			
+							TArray<FVector2f> UVsCoords;
+							int32 UVsIndex = 0;
+							while(UVsIndex < VertexUVs.Num())
+							{
+								if(VertexUVs[UVsIndex] != -1.0)
+								{
+									UVsCoords.Add(FVector2f(VertexUVs[UVsIndex], VertexUVs[UVsIndex+1]));
+									UVsIndex += 2;
+								}
+								else
+								{
+									UVsIndex += 1;
+								}
+							}
+							if(UVsCoords.Num() == VertexTextureUvs.Num())
+							{
+								for(int32 VertexIndex = 0; VertexIndex < UVsCoords.Num(); ++VertexIndex)
+								{
+									VertexTextureUvs[VertexIndex] = UVsCoords[VertexIndex];
+								}
+							}
+						}
+					}
+				}
+			}
+			SetValue(Context, MoveTemp(GroomCollection), &Collection);
+		}
+	}
+}
+
diff --git a/Engine/Plugins/Experimental/HairCardGenerator/Source/HairCardGeneratorDataflow/Private/HairCardDataflowRendering.cpp b/Engine/Plugins/Experimental/HairCardGenerator/Source/HairCardGeneratorDataflow/Private/HairCardDataflowRendering.cpp
new file mode 100644
--- /dev/null
+++ b/Engine/Plugins/Experimental/HairCardGenerator/Source/HairCardGeneratorDataflow/Private/HairCardDataflowRendering.cpp
@@ -0,0 +1,254 @@
+// Copyright Epic Games, Inc. All Rights Reserved.
+
+#include "HairCardDataflowRendering.h"
+#include "GenerateCardsGeometryNode.h"
+#include "GenerateCardsClumpsNode.h"
+#include "GenerateCardsTexturesNode.h"
+#include "GroomCollectionFacades.h"
+
+namespace UE::CardGen::Private
+{
+	int32 GGroomDataflowCardsLod = 0;
+	static FAutoConsoleVariableRef CVarGroomDataflowCardsLod(
+		 TEXT("p.Groom.Dataflow.CardsLod"),
+		 GGroomDataflowCardsLod,
+		 TEXT("Cards Lod we want to display")
+	 );
+
+	float GGroomDataflowCardsAlpha = 0.1;
+	static FAutoConsoleVariableRef CVarGroomDataflowCardsAlpha(
+		 TEXT("p.Groom.Dataflow.CardsAlpha"),
+		 GGroomDataflowCardsAlpha,
+		 TEXT("Cards Alpha for the rendering")
+	 );
+	
+	Dataflow::FRenderKey FCardsGeometryRenderingCallbacks::RenderKey = { "GeometryRender", FName("FCardsCollection") };
+	
+	void FCardsGeometryRenderingCallbacks::Render(GeometryCollection::Facades::FRenderingFacade& RenderCollection, const Dataflow::FGraphRenderingState& State)
+	{
+		if (State.GetRenderOutputs().Num())
+		{
+			const FManagedArrayCollection Default;
+			checkf(State.GetRenderOutputs().Num() == 1, TEXT("Expected FGraphRenderingState object to have one render output"));
+			const FName PrimaryOutput = State.GetRenderOutputs()[0];
+			FManagedArrayCollection Collection = State.GetValue<FManagedArrayCollection>(PrimaryOutput, Default);
+			
+			const TSharedRef<const FManagedArrayCollection> GroomCollection = MakeShared<FManagedArrayCollection>(MoveTemp(Collection));
+
+			FString CardsVerticesLODGroup = FGenerateCardsGeometryNode::CardsVerticesGroup.ToString();
+			CardsVerticesLODGroup.AppendInt(GGroomDataflowCardsLod);
+							
+			FString CardsFacesLODGroup = FGenerateCardsGeometryNode::CardsFacesGroup.ToString();
+			CardsFacesLODGroup.AppendInt(GGroomDataflowCardsLod);
+
+			if (GroomCollection->HasAttribute(FGenerateCardsGeometryNode::VertexClumpPositionsAttribute, FName(CardsVerticesLODGroup)) &&
+				GroomCollection->HasAttribute(FGenerateCardsGeometryNode::FaceVertexIndicesAttribute, FName(CardsFacesLODGroup)))
+			{
+				TArray<FVector3f> VertexPositions;
+				TArray<FVector3f> VertexNormals;
+				TArray<FIntVector> FaceIndices;
+			
+				if (State.GetViewMode().GetName() == Dataflow::FDataflowConstruction3DViewMode::Name)
+				{
+					const TManagedArray<FVector3f>& VertexPositionsArray = GroomCollection->GetAttribute<FVector3f>(
+						FGenerateCardsGeometryNode::VertexClumpPositionsAttribute, FName(CardsVerticesLODGroup));
+						
+					const TManagedArray<FIntVector3>& FaceIndicesArray = GroomCollection->GetAttribute<FIntVector3>(
+						FGenerateCardsGeometryNode::FaceVertexIndicesAttribute, FName(CardsFacesLODGroup));
+
+					VertexPositions = VertexPositionsArray.GetConstArray();
+					FaceIndices = FaceIndicesArray.GetConstArray();
+
+					TArray<int32> NumFaces;
+					
+					VertexNormals.Init(FVector3f::Zero(), VertexPositions.Num());
+					NumFaces.Init(0, VertexPositions.Num());
+					
+					for(int32 FaceIndex = 0; FaceIndex < FaceIndices.Num(); ++FaceIndex)
+					{
+						const FVector3f FaceNormal =  (VertexPositions[FaceIndices[FaceIndex][2]] - VertexPositions[FaceIndices[FaceIndex][0]]).Cross(
+							VertexPositions[FaceIndices[FaceIndex][1]] - VertexPositions[FaceIndices[FaceIndex][0]]).GetSafeNormal();
+						
+						VertexNormals[FaceIndices[FaceIndex][0]] += FaceNormal;
+						VertexNormals[FaceIndices[FaceIndex][1]] += FaceNormal;
+						VertexNormals[FaceIndices[FaceIndex][2]] += FaceNormal;
+
+						NumFaces[FaceIndices[FaceIndex][0]]++;
+						NumFaces[FaceIndices[FaceIndex][1]]++;
+						NumFaces[FaceIndices[FaceIndex][2]]++;
+					}
+
+					for(int32 VertexIndex = 0; VertexIndex < VertexPositions.Num(); ++VertexIndex)
+					{
+						if(NumFaces[VertexIndex] > 0)
+						{
+							VertexNormals[VertexIndex] /= NumFaces[VertexIndex];
+						}
+					}
+				}
+				else
+				{
+					checkf(false, TEXT("Invalid View Mode for FClothCollection rendering"));
+				}
+			
+				TArray<FLinearColor> Colors;
+				Colors.Init(FLinearColor::Gray, VertexPositions.Num());	// TODO: Choose a vertex color
+
+				ComputeVertexColors(GroomCollection, Colors);
+			
+				const int32 GeometryIndex = RenderCollection.StartGeometryGroup(State.GetGuid().ToString());
+				RenderCollection.AddSurface(MoveTemp(VertexPositions), MoveTemp(FaceIndices), MoveTemp(VertexNormals), MoveTemp(Colors));
+				RenderCollection.EndGeometryGroup(GeometryIndex);
+			}
+		}
+	}
+
+	Dataflow::FRenderKey FCardsTextureRenderingCallbacks::RenderKey = { "TextureRender", FName("FCardsCollection") };
+
+	void FCardsTextureRenderingCallbacks::Render(GeometryCollection::Facades::FRenderingFacade& RenderCollection, const Dataflow::FGraphRenderingState& State)
+	{
+		if (State.GetViewMode().GetName() == Dataflow::FDataflowConstruction3DViewMode::Name)
+		{
+			FCardsGeometryRenderingCallbacks::Render(RenderCollection, State);
+		}
+		else if (State.GetViewMode().GetName() == Dataflow::FDataflowConstructionUVViewMode::Name)
+		{
+			if (State.GetRenderOutputs().Num())
+			{
+				const FManagedArrayCollection Default;
+				checkf(State.GetRenderOutputs().Num() == 1, TEXT("Expected FGraphRenderingState object to have one render output"));
+				const FName PrimaryOutput = State.GetRenderOutputs()[0];
+				FManagedArrayCollection Collection = State.GetValue<FManagedArrayCollection>(PrimaryOutput, Default);
+			
+				const TSharedRef<const FManagedArrayCollection> GroomCollection = MakeShared<FManagedArrayCollection>(MoveTemp(Collection));
+
+				FString CardsVerticesLODGroup = FGenerateCardsGeometryNode::CardsVerticesGroup.ToString();
+				CardsVerticesLODGroup.AppendInt(GGroomDataflowCardsLod);
+							
+				FString CardsFacesLODGroup = FGenerateCardsGeometryNode::CardsFacesGroup.ToString();
+				CardsFacesLODGroup.AppendInt(GGroomDataflowCardsLod);
+
+				if (GroomCollection->HasAttribute(FGenerateCardsTexturesNode::VertexTextureUVsAttribute, FName(CardsVerticesLODGroup)) &&
+					GroomCollection->HasAttribute(FGenerateCardsGeometryNode::FaceVertexIndicesAttribute, FName(CardsFacesLODGroup)))
+				{
+					const TManagedArray<FVector2f>& VertexUVsArray = GroomCollection->GetAttribute<FVector2f>(
+							FGenerateCardsTexturesNode::VertexTextureUVsAttribute, FName(CardsVerticesLODGroup));
+						
+					const TManagedArray<FIntVector3>& FaceIndicesArray = GroomCollection->GetAttribute<FIntVector3>(
+						FGenerateCardsGeometryNode::FaceVertexIndicesAttribute, FName(CardsFacesLODGroup));
+					
+					TArray<FIntVector> FaceIndices = FaceIndicesArray.GetConstArray();
+					
+					TArray<FVector3f> VertexUVs;
+					VertexUVs.Init(FVector3f(0.0, 0.0, 1.0), VertexUVsArray.Num());
+
+					for (int32 VertexIndex = 0; VertexIndex < VertexUVsArray.Num(); ++VertexIndex)
+					{
+						VertexUVs[VertexIndex] = FVector3f(VertexUVsArray[VertexIndex][0], VertexUVsArray[VertexIndex][1], 0.0f);
+					}
+
+					TArray<FVector3f> VertexNormals;
+					VertexNormals.Init(FVector3f(0.0, 0.0, 1.0), VertexUVs.Num());
+
+					TArray<FLinearColor> VertexColors;
+					VertexColors.Init(FLinearColor(0,0,0,0), VertexUVs.Num());
+
+					const int32 GeometryIndex = RenderCollection.StartGeometryGroup(State.GetGuid().ToString());
+					RenderCollection.AddSurface(MoveTemp(VertexUVs), MoveTemp(FaceIndices), MoveTemp(VertexNormals), MoveTemp(VertexColors));
+					RenderCollection.EndGeometryGroup(GeometryIndex);
+				}	
+			}
+		}
+	}
+	
+	void FCardsTextureRenderingCallbacks::ComputeVertexColors(const TSharedRef<const FManagedArrayCollection> GroomCollection, TArray<FLinearColor>& VertexColors) const
+	{
+		FString CardsObjectsLODGroup = FGenerateCardsTexturesNode::CardsObjectsGroup.ToString();
+		CardsObjectsLODGroup.AppendInt(GGroomDataflowCardsLod);
+
+		FString CardsVerticesLODGroup = FGenerateCardsGeometryNode::CardsVerticesGroup.ToString();
+		CardsVerticesLODGroup.AppendInt(GGroomDataflowCardsLod);
+
+		if (GroomCollection->HasAttribute(FGenerateCardsTexturesNode::ObjectTextureIndicesAttribute, FName(CardsObjectsLODGroup)) &&
+			GroomCollection->HasAttribute(FGenerateCardsGeometryNode::VertexCardIndicesAttribute,FName(CardsVerticesLODGroup)))
+		{
+			const TManagedArray<int32>& VertexCardsArray = GroomCollection->GetAttribute<int32>(
+					FGenerateCardsGeometryNode::VertexCardIndicesAttribute, FName(CardsVerticesLODGroup));
+						
+			const TManagedArray<int32>& CardsTextureArray = GroomCollection->GetAttribute<int32>(
+					FGenerateCardsTexturesNode::ObjectTextureIndicesAttribute, FName(CardsObjectsLODGroup));
+
+			int32 NumTextures = 0;
+			for(int32 CardTexture = 0; CardTexture < CardsTextureArray.Num(); ++CardTexture)
+			{
+				NumTextures = FMath::Max(NumTextures, CardsTextureArray[CardTexture]);
+			}
+			NumTextures += 1;
+
+			TArray<FLinearColor> TextureColors;
+			TextureColors.Init(FLinearColor::Black, NumTextures);
+
+			for(int32 TextureIndex = 0; TextureIndex < NumTextures; ++TextureIndex)
+			{
+				TextureColors[TextureIndex] = FLinearColor::MakeRandomColor();
+			}
+
+			const int32 NumVertices = VertexCardsArray.Num();
+			VertexColors.Init(FLinearColor::Black, NumVertices);
+			
+			for(int32 VertexIndex = 0; VertexIndex < NumVertices; ++VertexIndex)
+			{
+				const int32 TextureIndex = (VertexCardsArray[VertexIndex] < CardsTextureArray.Num()) ?
+					CardsTextureArray[VertexCardsArray[VertexIndex]] : INDEX_NONE;
+				VertexColors[VertexIndex] = (TextureIndex != INDEX_NONE) ? TextureColors[TextureIndex] : FLinearColor::Black;
+				VertexColors[VertexIndex] *= (TextureIndex == VertexCardsArray[VertexIndex]) ? 1.0 : GGroomDataflowCardsAlpha;
+			}
+		}
+	}
+
+	Dataflow::FRenderKey FCardsClumpsRenderingCallbacks::RenderKey = { "ClumpsRender", FName("FCardsCollection") };
+	
+	void FCardsClumpsRenderingCallbacks::ComputeVertexColors(const UE::Groom::FGroomStrandsFacade& StrandsFacade, TArray<FLinearColor>& VertexColors) const
+	{
+		const FManagedArrayCollection& GroomCollection = StrandsFacade.GetManagedArrayCollection();
+
+		FString ClumpIndicesLOD = FGenerateCardsClumpsNode::CurveClumpIndicesAttribute.ToString();
+		ClumpIndicesLOD.AppendInt(GGroomDataflowCardsLod);
+
+		FString NumClumpsLOD = FGenerateCardsClumpsNode::ObjectNumClumpsAttribute.ToString();
+		NumClumpsLOD.AppendInt(GGroomDataflowCardsLod);
+		
+		if(GroomCollection.HasAttribute(FName(ClumpIndicesLOD), UE::Groom::FGroomStrandsFacade::CurvesGroup) &&
+		   GroomCollection.HasAttribute(FName(NumClumpsLOD), UE::Groom::FGroomStrandsFacade::ObjectsGroup))
+		{
+			const TManagedArray<int32>& ClumpIndices = GroomCollection.GetAttribute<int32>(FName(ClumpIndicesLOD), UE::Groom::FGroomStrandsFacade::CurvesGroup);
+			const TManagedArray<int32>& NumClumps = GroomCollection.GetAttribute<int32>(FName(NumClumpsLOD), UE::Groom::FGroomStrandsFacade::ObjectsGroup);
+
+			TArray<FLinearColor> ClumpsColors;
+			ClumpsColors.Init(FLinearColor::Black, NumClumps[0]);
+
+			for(int32 ClumpIndex = 0; ClumpIndex < NumClumps[0]; ++ClumpIndex)
+			{
+				ClumpsColors[ClumpIndex] = FLinearColor::MakeRandomColor();
+			}
+
+			TArray<FLinearColor> StrandsColors;
+			StrandsColors.Init(FLinearColor::Black, StrandsFacade.GetNumCurves());
+			for(int32 CurveIndex = 0, CurveEnd = StrandsFacade.GetNumCurves(); CurveIndex < CurveEnd; ++CurveIndex)
+			{
+				if(ClumpIndices[CurveIndex] < ClumpsColors.Num())
+				{
+					StrandsColors[CurveIndex] = ClumpsColors[ClumpIndices[CurveIndex]];
+				}
+			}
+			
+			for(int32 PointIndex = 0, PointEnd = StrandsFacade.GetNumPoints(); PointIndex < PointEnd; ++PointIndex)
+			{
+				const int32 CurveIndex = StrandsFacade.GetPointCurveIndices()[PointIndex];
+
+				VertexColors[2*PointIndex] = StrandsColors[CurveIndex];
+				VertexColors[2*PointIndex+1] = StrandsColors[CurveIndex];
+			}
+		}
+	}
+}
diff --git a/Engine/Plugins/Experimental/HairCardGenerator/Source/HairCardGeneratorDataflow/Private/HairCardGeneratorDataflowModule.cpp b/Engine/Plugins/Experimental/HairCardGenerator/Source/HairCardGeneratorDataflow/Private/HairCardGeneratorDataflowModule.cpp
new file mode 100644
--- /dev/null
+++ b/Engine/Plugins/Experimental/HairCardGenerator/Source/HairCardGeneratorDataflow/Private/HairCardGeneratorDataflowModule.cpp
@@ -0,0 +1,32 @@
+// Copyright Epic Games, Inc. All Rights Reserved.
+
+#include "HairCardGeneratorDataflowModule.h"
+#include "GenerateCardsClumpsNode.h"
+#include "BuildCardsSettingsNode.h"
+#include "GenerateCardsGeometryNode.h"
+#include "GenerateCardsTexturesNode.h"
+#include "CardsAssetTerminalNode.h"
+#include "HairCardDataflowRendering.h"
+#include "Modules/ModuleManager.h"
+
+#define LOCTEXT_NAMESPACE "HairCardGeneratorDataflow"
+
+void FHairCardGeneratorDataflowModule::StartupModule()
+{
+	DATAFLOW_NODE_REGISTER_CREATION_FACTORY(FGenerateCardsClumpsNode);
+	DATAFLOW_NODE_REGISTER_CREATION_FACTORY(FGenerateCardsGeometryNode);
+	DATAFLOW_NODE_REGISTER_CREATION_FACTORY(FGenerateCardsTexturesNode);
+	DATAFLOW_NODE_REGISTER_CREATION_FACTORY(FBuildCardsSettingsNode);
+	DATAFLOW_NODE_REGISTER_CREATION_FACTORY(FCardsAssetTerminalNode);
+
+	UE::CardGen::Private::RegisterRenderingCallbacks();
+}
+
+void FHairCardGeneratorDataflowModule::ShutdownModule()
+{
+	UE::CardGen::Private::DeregisterRenderingCallbacks();
+}
+
+IMPLEMENT_MODULE(FHairCardGeneratorDataflowModule, HairCardGeneratorDataflow)
+
+#undef LOCTEXT_NAMESPACE
diff --git a/Engine/Plugins/Experimental/HairCardGenerator/Source/HairCardGeneratorDataflow/Public/BuildCardsSettingsNode.h b/Engine/Plugins/Experimental/HairCardGenerator/Source/HairCardGeneratorDataflow/Public/BuildCardsSettingsNode.h
new file mode 100644
--- /dev/null
+++ b/Engine/Plugins/Experimental/HairCardGenerator/Source/HairCardGeneratorDataflow/Public/BuildCardsSettingsNode.h
@@ -0,0 +1,70 @@
+// Copyright Epic Games, Inc. All Rights Reserved.
+
+#pragma once 
+
+#include "CoreMinimal.h"
+#include "GroomAsset.h"
+#include "ChaosLog.h"
+#include "HairCardGeneratorPluginSettings.h"
+#include "Dataflow/DataflowCore.h"
+#include "Dataflow/DataflowEngine.h"
+
+#include "BuildCardsSettingsNode.generated.h"
+
+USTRUCT()
+struct FGroomCardsSettings
+{
+	GENERATED_BODY()
+	
+	/** Generator settings to be built */
+	UPROPERTY(EditAnywhere, Category = "GroomCards");
+	TObjectPtr<UHairCardGeneratorPluginSettings> GenerationSettings = nullptr;
+
+	/** Generation flags to output the assets */
+	UPROPERTY(EditAnywhere, Category = "GroomCards");
+	uint8 GenerationFlags = 0;
+
+	/** Pipeline flags to generate clumps, geometry and textures */
+	UPROPERTY(EditAnywhere, Category = "GroomCards");
+	uint8 PipelineFlags = 0;
+
+	/** Groom asset to generate the */
+    UPROPERTY(EditAnywhere, Category = "GroomCards");
+    TObjectPtr<UGroomAsset> GroomAsset = nullptr;
+};
+
+/** Build the cards generation settings */
+USTRUCT(meta = (DataflowGroom))
+struct FBuildCardsSettingsNode : public FDataflowNode
+{
+	GENERATED_BODY()
+	
+	DATAFLOW_NODE_DEFINE_INTERNAL(FBuildCardsSettingsNode, "BuildCardsSettings", "Groom", "")
+
+public:
+	
+	FBuildCardsSettingsNode(const UE::Dataflow::FNodeParameters& InParam, FGuid InGuid = FGuid::NewGuid())
+	: FDataflowNode(InParam, InGuid)
+	{
+		RegisterInputConnection(&Collection);
+		RegisterOutputConnection(&Collection, &Collection);
+		RegisterOutputConnection(&CardsSettings);
+	}
+
+	//~ Begin FDataflowNode interface
+	virtual void Evaluate(UE::Dataflow::FContext& Context, const FDataflowOutput* Out) const override;
+	//~ End FDataflowNode interface
+
+	/** Groom asset to build the cards settings from */
+	UPROPERTY(EditAnywhere, Category = "GroomCards")
+	TObjectPtr<UGroomAsset> GroomAsset;
+	
+	/** Managed array collection to be used to store datas */
+	UPROPERTY(meta = (DataflowInput, DataflowOutput, DisplayName = "Collection", DataflowPassthrough  = "Collection"))
+	FManagedArrayCollection Collection;
+
+	/** Generator clumps settings to be built */
+	UPROPERTY(meta = (DataflowOutput));
+	TArray<FGroomCardsSettings> CardsSettings;
+};
+
diff --git a/Engine/Plugins/Experimental/HairCardGenerator/Source/HairCardGeneratorDataflow/Public/CardsAssetTerminalNode.h b/Engine/Plugins/Experimental/HairCardGenerator/Source/HairCardGeneratorDataflow/Public/CardsAssetTerminalNode.h
new file mode 100644
--- /dev/null
+++ b/Engine/Plugins/Experimental/HairCardGenerator/Source/HairCardGeneratorDataflow/Public/CardsAssetTerminalNode.h
@@ -0,0 +1,42 @@
+// Copyright Epic Games, Inc. All Rights Reserved.
+
+#pragma once 
+
+#include "CoreMinimal.h"
+#include "BuildCardsSettingsNode.h"
+#include "GroomAsset.h"
+#include "ChaosLog.h"
+#include "Dataflow/DataflowCore.h"
+#include "Dataflow/DataflowEngine.h"
+
+#include "CardsAssetTerminalNode.generated.h"
+
+USTRUCT(meta = (DataflowGroom, DataflowTerminal))
+struct FCardsAssetTerminalNode : public FDataflowTerminalNode
+{
+	GENERATED_USTRUCT_BODY()
+	DATAFLOW_NODE_DEFINE_INTERNAL(FCardsAssetTerminalNode, "CardsAssetTerminal", "Groom", "")
+
+public:
+
+	FCardsAssetTerminalNode(const UE::Dataflow::FNodeParameters& InParam, FGuid InGuid = FGuid::NewGuid())
+	: FDataflowTerminalNode(InParam, InGuid)
+	{
+		RegisterInputConnection(&Collection);
+		RegisterInputConnection(&CardsSettings);
+		RegisterOutputConnection(&Collection, &Collection);
+	}
+
+	//~ Begin FDataflowTerminalNode interface
+	virtual void SetAssetValue(TObjectPtr<UObject> Asset, UE::Dataflow::FContext& Context) const override;
+	virtual void Evaluate(UE::Dataflow::FContext& Context, const FDataflowOutput* Out) const override;
+	//~ End FDataflowTerminalNode interface
+
+	UPROPERTY(meta = (DataflowInput, DataflowOutput, DisplayName = "Collection", DataflowPassthrough  = "Collection"))
+	FManagedArrayCollection Collection;
+
+	/** Generator settings to be used */
+	UPROPERTY(meta = (DataflowInput, DisplayName = "CardsSettings"));
+	TArray<FGroomCardsSettings> CardsSettings;
+};
+
diff --git a/Engine/Plugins/Experimental/HairCardGenerator/Source/HairCardGeneratorDataflow/Public/GenerateCardsClumpsNode.h b/Engine/Plugins/Experimental/HairCardGenerator/Source/HairCardGeneratorDataflow/Public/GenerateCardsClumpsNode.h
new file mode 100644
--- /dev/null
+++ b/Engine/Plugins/Experimental/HairCardGenerator/Source/HairCardGeneratorDataflow/Public/GenerateCardsClumpsNode.h
@@ -0,0 +1,75 @@
+// Copyright Epic Games, Inc. All Rights Reserved.
+
+#pragma once 
+
+#include "CoreMinimal.h"
+#include "GroomAsset.h"
+#include "ChaosLog.h"
+#include "BuildCardsSettingsNode.h"
+#include "Dataflow/DataflowCore.h"
+#include "Dataflow/DataflowEngine.h"
+
+#include "GenerateCardsClumpsNode.generated.h"
+
+USTRUCT()
+struct FCardsClumpsSettings
+{
+	GENERATED_BODY()
+	
+	// LOD index for which we will override the number of cards/triangles/texture/flyaways
+	UPROPERTY(EditAnywhere, Category="GroomCards")
+	int32 LODIndex = -1;
+
+	// Card group name for which we will override the number of cards/triangles/texture/flyaways
+	UPROPERTY(EditAnywhere, Category="GroomCards")
+	FName CardGroup = NAME_None;
+
+	// Total number of cards to generate for this LOD settings
+	UPROPERTY(EditAnywhere, Category="GroomCards", meta=(ClampMin="1", ClampMax="100000"))
+	int32 NumCards = 200;
+
+	// Maximum number of cards to assign to flyaway strands
+	UPROPERTY(EditAnywhere, Category="GroomCards", meta=(ClampMin="0", ClampMax="1000"))
+	int32 NumFlyaways = 10;
+};
+
+/** Generate the textures used to build the cards from the strands */
+USTRUCT(meta = (DataflowGroom))
+struct FGenerateCardsClumpsNode : public FDataflowNode
+{
+	GENERATED_BODY()
+	
+	DATAFLOW_NODE_DEFINE_INTERNAL(FGenerateCardsClumpsNode, "GenerateCardsClumps", "Groom", "")
+	DATAFLOW_NODE_RENDER_TYPE("ClumpsRender", FName("FCardsCollection"), "Collection")
+
+public:
+
+	static const FName CurveClumpIndicesAttribute;
+	static const FName ObjectNumClumpsAttribute;
+	
+	FGenerateCardsClumpsNode(const UE::Dataflow::FNodeParameters& InParam, FGuid InGuid = FGuid::NewGuid())
+	: FDataflowNode(InParam, InGuid)
+	{
+		RegisterInputConnection(&Collection);
+		RegisterInputConnection(&CardsSettings);
+		RegisterOutputConnection(&Collection, &Collection);
+		RegisterOutputConnection(&CardsSettings, &CardsSettings);
+	}
+
+	//~ Begin FDataflowNode interface
+	virtual void Evaluate(UE::Dataflow::FContext& Context, const FDataflowOutput* Out) const override;
+	//~ End FDataflowNode interface
+	
+	/** Per LOD settings for each cards generation */
+	UPROPERTY(EditAnywhere, Category = "GroomCards")
+	TArray<FCardsClumpsSettings> ClumpsSettings;
+
+	/** Managed array collection to be used to store datas */
+	UPROPERTY(meta = (DataflowInput, DataflowOutput, DisplayName = "Collection", DataflowPassthrough  = "Collection"))
+	FManagedArrayCollection Collection;
+
+	/** Generator settings to be used */
+	UPROPERTY(meta = (DataflowInput, DataflowOutput, DisplayName = "Cards Settings", DataflowPassthrough  = "CardsSettings"));
+	TArray<FGroomCardsSettings> CardsSettings;
+};
+
diff --git a/Engine/Plugins/Experimental/HairCardGenerator/Source/HairCardGeneratorDataflow/Public/GenerateCardsGeometryNode.h b/Engine/Plugins/Experimental/HairCardGenerator/Source/HairCardGeneratorDataflow/Public/GenerateCardsGeometryNode.h
new file mode 100644
--- /dev/null
+++ b/Engine/Plugins/Experimental/HairCardGenerator/Source/HairCardGeneratorDataflow/Public/GenerateCardsGeometryNode.h
@@ -0,0 +1,75 @@
+// Copyright Epic Games, Inc. All Rights Reserved.
+
+#pragma once 
+
+#include "CoreMinimal.h"
+#include "GroomAsset.h"
+#include "ChaosLog.h"
+#include "BuildCardsSettingsNode.h"
+#include "Dataflow/DataflowCore.h"
+#include "Dataflow/DataflowEngine.h"
+
+#include "GenerateCardsGeometryNode.generated.h"
+
+USTRUCT()
+struct FCardsGeometrySettings
+{
+	GENERATED_BODY()
+	
+	// LOD index for which we will override the number of cards/triangles/texture/flyaways
+	UPROPERTY(EditAnywhere, Category="GroomCards")
+	int32 LODIndex = -1;
+
+	// Card group name for which we will override the number of cards/triangles/texture/flyaways
+	UPROPERTY(EditAnywhere, Category="GroomCards")
+	FName CardGroup = NAME_None;
+
+	// Total number of triangles to generate for this LOD settings
+	UPROPERTY(EditAnywhere, Category="GroomCards", meta=(ClampMin="1", ClampMax="100000"))
+	int32 NumTriangles = 2000;
+};
+
+/** Generate the geometry used to build the cards from the strands */
+USTRUCT(meta = (DataflowGroom))
+struct FGenerateCardsGeometryNode : public FDataflowNode
+{
+	GENERATED_BODY()
+	
+	DATAFLOW_NODE_DEFINE_INTERNAL(FGenerateCardsGeometryNode, "GenerateCardsGeometry", "Groom", "")
+	DATAFLOW_NODE_RENDER_TYPE("GeometryRender", FName("FCardsCollection"), "Collection")
+
+public:
+
+	static const FName VertexClumpPositionsAttribute;
+	static const FName FaceVertexIndicesAttribute;
+	static const FName VertexCardIndicesAttribute;
+	
+	static const FName CardsVerticesGroup;
+	static const FName CardsFacesGroup;
+	
+	FGenerateCardsGeometryNode(const UE::Dataflow::FNodeParameters& InParam, FGuid InGuid = FGuid::NewGuid())
+	: FDataflowNode(InParam, InGuid)
+	{
+		RegisterInputConnection(&Collection);
+		RegisterInputConnection(&CardsSettings);
+		RegisterOutputConnection(&Collection, &Collection);
+		RegisterOutputConnection(&CardsSettings, &CardsSettings);
+	}
+
+	//~ Begin FDataflowNode interface
+	virtual void Evaluate(UE::Dataflow::FContext& Context, const FDataflowOutput* Out) const override;
+	//~ End FDataflowNode interface
+
+	/** Per LOD settings for each cards generation */
+	UPROPERTY(EditAnywhere, Category = "GroomCards")
+	TArray<FCardsGeometrySettings> GeometrySettings;
+
+	/** Managed array collection to be used to store datas */
+	UPROPERTY(meta = (DataflowInput, DataflowOutput, DisplayName = "Collection", DataflowPassthrough  = "Collection"))
+	FManagedArrayCollection Collection;
+
+	/** Generator settings to be used */
+	UPROPERTY(meta = (DataflowInput, DataflowOutput, DisplayName = "CardsSettings", DataflowPassthrough  = "CardsSettings"));
+	TArray<FGroomCardsSettings> CardsSettings;
+};
+
diff --git a/Engine/Plugins/Experimental/HairCardGenerator/Source/HairCardGeneratorDataflow/Public/GenerateCardsTexturesNode.h b/Engine/Plugins/Experimental/HairCardGenerator/Source/HairCardGeneratorDataflow/Public/GenerateCardsTexturesNode.h
new file mode 100644
--- /dev/null
+++ b/Engine/Plugins/Experimental/HairCardGenerator/Source/HairCardGeneratorDataflow/Public/GenerateCardsTexturesNode.h
@@ -0,0 +1,73 @@
+// Copyright Epic Games, Inc. All Rights Reserved.
+
+#pragma once 
+
+#include "CoreMinimal.h"
+#include "GroomAsset.h"
+#include "ChaosLog.h"
+#include "BuildCardsSettingsNode.h"
+#include "Dataflow/DataflowCore.h"
+#include "Dataflow/DataflowEngine.h"
+
+#include "GenerateCardsTexturesNode.generated.h"
+
+USTRUCT()
+struct FCardsTextureSettings
+{
+	GENERATED_BODY()
+	
+	// LOD index for which we will override the number of cards/triangles/texture/flyaways
+	UPROPERTY(EditAnywhere, Category="GroomCards")
+	int32 LODIndex = -1;
+
+	// Card group name for which we will override the number of cards/triangles/texture/flyaways
+	UPROPERTY(EditAnywhere, Category="GroomCards")
+	FName CardGroup = NAME_None;
+
+	// Total number of triangles to generate for this LOD settings
+	UPROPERTY(EditAnywhere, Category="GroomCards")
+	int32 NumTextures = 75;
+};
+
+/** Generate the clumps used to build the cards from the strands */
+USTRUCT(meta = (DataflowGroom))
+struct FGenerateCardsTexturesNode : public FDataflowNode
+{
+	GENERATED_BODY()
+	
+	DATAFLOW_NODE_DEFINE_INTERNAL(FGenerateCardsTexturesNode, "GenerateCardsTextures", "Groom", "")
+	DATAFLOW_NODE_RENDER_TYPE("TextureRender", FName("FCardsCollection"), "Collection")
+
+public:
+
+
+	static const FName CardsObjectsGroup;
+	static const FName ObjectTextureIndicesAttribute;
+	static const FName VertexTextureUVsAttribute;
+	
+	FGenerateCardsTexturesNode(const UE::Dataflow::FNodeParameters& InParam, FGuid InGuid = FGuid::NewGuid())
+	: FDataflowNode(InParam, InGuid)
+	{
+		RegisterInputConnection(&Collection);
+		RegisterInputConnection(&CardsSettings);
+		RegisterOutputConnection(&Collection, &Collection);
+		RegisterOutputConnection(&CardsSettings, &CardsSettings);
+	}
+
+	//~ Begin FDataflowNode interface
+	virtual void Evaluate(UE::Dataflow::FContext& Context, const FDataflowOutput* Out) const override;
+	//~ End FDataflowNode interface
+
+	/** Per LOD settings for each cards generation */
+	UPROPERTY(EditAnywhere, Category = "GroomCards")
+	TArray<FCardsTextureSettings> TextureSettings;
+
+	/** Managed array collection to be used to store datas */
+	UPROPERTY(meta = (DataflowInput, DataflowOutput, DisplayName = "Collection", DataflowPassthrough  = "Collection"))
+	FManagedArrayCollection Collection;
+
+	/** Generator settings to be used */
+	UPROPERTY(meta = (DataflowInput, DataflowOutput, DisplayName = "CardsSettings", DataflowPassthrough  = "CardsSettings"));
+	TArray<FGroomCardsSettings> CardsSettings;
+};
+
diff --git a/Engine/Plugins/Experimental/HairCardGenerator/Source/HairCardGeneratorDataflow/Public/HairCardDataflowRendering.h b/Engine/Plugins/Experimental/HairCardGenerator/Source/HairCardGeneratorDataflow/Public/HairCardDataflowRendering.h
new file mode 100644
--- /dev/null
+++ b/Engine/Plugins/Experimental/HairCardGenerator/Source/HairCardGeneratorDataflow/Public/HairCardDataflowRendering.h
@@ -0,0 +1,79 @@
+// Copyright Epic Games, Inc. All Rights Reserved.
+
+#include "GroomDataflowRendering.h"
+#include "Dataflow/DataflowRenderingFactory.h"
+#include "Dataflow/DataflowRenderingViewMode.h"
+
+namespace UE::CardGen::Private
+{
+	class FCardsGeometryRenderingCallbacks : public Dataflow::FRenderingFactory::ICallbackInterface
+	{
+	public :
+		static Dataflow::FRenderKey RenderKey;
+		
+	protected:
+
+		virtual Dataflow::FRenderKey GetRenderKey() const override
+		{
+			return RenderKey;
+		}
+
+		virtual bool CanRender(const Dataflow::IDataflowConstructionViewMode& ViewMode) const override
+		{
+			return (ViewMode.GetName() == Dataflow::FDataflowConstruction3DViewMode::Name);
+		}
+
+		virtual void Render(GeometryCollection::Facades::FRenderingFacade& RenderCollection, const Dataflow::FGraphRenderingState& State) override;
+		virtual void ComputeVertexColors(const TSharedRef<const FManagedArrayCollection> GroomCollection, TArray<FLinearColor>& VertexColors) const {}
+	};
+
+	class FCardsClumpsRenderingCallbacks : public UE::Groom::FGroomStrandsRenderingCallbacks
+	{
+	public :
+		static Dataflow::FRenderKey RenderKey;
+		
+	private:
+
+		virtual Dataflow::FRenderKey GetRenderKey() const override
+		{
+			return RenderKey;
+		}
+		
+		virtual void ComputeVertexColors(const UE::Groom::FGroomStrandsFacade& StrandsFacade, TArray<FLinearColor>& VertexColors) const override;
+	};
+
+	class FCardsTextureRenderingCallbacks : public FCardsGeometryRenderingCallbacks
+	{
+	public :
+			static Dataflow::FRenderKey RenderKey;
+		
+	private:
+
+		virtual bool CanRender(const Dataflow::IDataflowConstructionViewMode& ViewMode) const override
+		{
+			return (ViewMode.GetName() == Dataflow::FDataflowConstruction3DViewMode::Name || ViewMode.GetName() == Dataflow::FDataflowConstructionUVViewMode::Name);
+		}
+
+		virtual Dataflow::FRenderKey GetRenderKey() const override
+		{
+			return RenderKey;
+		}
+		virtual void Render(GeometryCollection::Facades::FRenderingFacade& RenderCollection, const Dataflow::FGraphRenderingState& State) override;
+		virtual void ComputeVertexColors(const TSharedRef<const FManagedArrayCollection> GroomCollection, TArray<FLinearColor>& VertexColors) const override;
+	};
+	
+	static void RegisterRenderingCallbacks()
+	{
+		Dataflow::FRenderingFactory::GetInstance()->RegisterCallbacks(MakeUnique<FCardsClumpsRenderingCallbacks>());
+		Dataflow::FRenderingFactory::GetInstance()->RegisterCallbacks(MakeUnique<FCardsGeometryRenderingCallbacks>());
+		Dataflow::FRenderingFactory::GetInstance()->RegisterCallbacks(MakeUnique<FCardsTextureRenderingCallbacks>());
+	}
+
+	static void DeregisterRenderingCallbacks()
+	{
+		Dataflow::FRenderingFactory::GetInstance()->DeregisterCallbacks(FCardsGeometryRenderingCallbacks::RenderKey);
+		Dataflow::FRenderingFactory::GetInstance()->DeregisterCallbacks(FCardsClumpsRenderingCallbacks::RenderKey);
+		Dataflow::FRenderingFactory::GetInstance()->DeregisterCallbacks(FCardsTextureRenderingCallbacks::RenderKey);
+	}
+
+}  // End namespace Private
diff --git a/Engine/Plugins/Experimental/HairCardGenerator/Source/HairCardGeneratorDataflow/Public/HairCardGeneratorDataflowModule.h b/Engine/Plugins/Experimental/HairCardGenerator/Source/HairCardGeneratorDataflow/Public/HairCardGeneratorDataflowModule.h
new file mode 100644
--- /dev/null
+++ b/Engine/Plugins/Experimental/HairCardGenerator/Source/HairCardGeneratorDataflow/Public/HairCardGeneratorDataflowModule.h
@@ -0,0 +1,17 @@
+// Copyright Epic Games, Inc. All Rights Reserved.
+
+#pragma once
+
+#include "CoreMinimal.h"
+#include "Modules/ModuleInterface.h"
+
+/** Hair card generator dataflow module public interface */
+class FHairCardGeneratorDataflowModule : public IModuleInterface
+{
+	public :
+	
+	//~ IModuleInterface begin interface
+	virtual void StartupModule() override;
+	virtual void ShutdownModule() override;
+	//~ IModuleInterface end interface
+};
diff --git a/Engine/Plugins/Experimental/HairCardGenerator/Source/HairCardGeneratorEditor/Private/HairCardGenControllerBase.h b/Engine/Plugins/Experimental/HairCardGenerator/Source/HairCardGeneratorEditor/Private/HairCardGenControllerBase.h
--- a/Engine/Plugins/Experimental/HairCardGenerator/Source/HairCardGeneratorEditor/Private/HairCardGenControllerBase.h
+++ b/Engine/Plugins/Experimental/HairCardGenerator/Source/HairCardGeneratorEditor/Private/HairCardGenControllerBase.h
@@ -74,10 +74,10 @@ public:
 	bool LoadSettings(const UHairCardGeneratorPluginSettings* GeneratorSettings);
 
 	UFUNCTION(BlueprintImplementableEvent, Category = Python)
-	bool GenerateCardsGeometry();
+	TArray<float> GenerateCardsGeometry();
 
 	UFUNCTION(BlueprintImplementableEvent, Category = Python)
-	bool GenerateClumps(const int GroupId = -1);
+	TArray<int32> GenerateClumps(const int GroupId = -1);
 
 	UFUNCTION(BlueprintImplementableEvent, Category = Python)
 	TArray<int> SetOptimizations(const int GroupId = -1);
@@ -89,10 +89,10 @@ public:
 	void SetInterpolatedAvgCurve(const int Id, const int Cid, const TArray<float>& Points);
 
 	UFUNCTION(BlueprintImplementableEvent, Category = Python)
-	bool ClusterTextures(int GroupIndex);
+	TArray<int32> ClusterTextures(int GroupIndex);
 
 	UFUNCTION(BlueprintImplementableEvent, Category = Python)
-	bool GenerateTextureLayout();
+	TArray<float> GenerateTextureLayout();
 
 	UFUNCTION(BlueprintImplementableEvent, Category = Python)
 	bool GenerateTextureAtlases(const float WidthScale = -1);
diff --git a/Engine/Plugins/Experimental/HairCardGenerator/Source/HairCardGeneratorEditor/Private/HairCardGeneratorEditorModule.cpp b/Engine/Plugins/Experimental/HairCardGenerator/Source/HairCardGeneratorEditor/Private/HairCardGeneratorEditorModule.cpp
--- a/Engine/Plugins/Experimental/HairCardGenerator/Source/HairCardGeneratorEditor/Private/HairCardGeneratorEditorModule.cpp
+++ b/Engine/Plugins/Experimental/HairCardGenerator/Source/HairCardGeneratorEditor/Private/HairCardGeneratorEditorModule.cpp
@@ -1,5 +1,6 @@
 // Copyright Epic Games, Inc. All Rights Reserved.
 
+#include "HairCardGeneratorEditorModule.h"
 #include "IHairCardGeneratorEditor.h"
 #include "Engine/Texture2D.h"
 #include "Factories/TextureFactory.h"
@@ -47,7 +48,6 @@
 	TArray<FString> TextureNames;
 };
 
-
 /* FHairCardGeneratorHelper Declaration
  *****************************************************************************/
 
@@ -98,10 +98,10 @@
 
 	static void DbgCmdCacheGroom(const TArray<FString>& Args);
 
-	static bool LoadGroomData(UGroomAsset* GroomAsset, const bool SaveCached);
+	static bool LoadGroomData(UGroomAsset* GroomAsset, const bool SaveCached = false);
 	static bool BuildHairCardGenGroomData(const int PointsPerStrand, const FHairDescription& InGroomData, FHairCardGen_GroomData& Out);
-	static bool GenerateCardsForCardGroup(TObjectPtr<const UHairCardGeneratorPluginSettings> Settings, int index, uint8 GenFlags);
-	static bool GenerateCardsGeometry(TObjectPtr<const UHairCardGeneratorPluginSettings> Settings, int index);
+	static bool GenerateCardsForCardGroup(const TObjectPtr<const UHairCardGeneratorPluginSettings> GenerationSettings, const int32 FilterIndex, const uint8 GenFlags);
+	static bool GenerateCardsGeometry(TObjectPtr<const UHairCardGeneratorPluginSettings> Settings, int index, TArray<TArray<FVector3f>>& ClumpsGeometry);
 	static void GetTextureImportNames(UHairCardGeneratorPluginSettings* Settings, FHairCardImportSettings& ImportSettings);
 	static TArray<UObject*> ImportHairCardTextures(const FString& ImportSrc, const FString& ImportDst, const FHairCardImportSettings& ImportSettings);
 	static UStaticMesh* CreateMeshInPackage(const FString& DestinationPackage, const FString& BaseFilename, int GroupID);
@@ -117,5 +117,376 @@
 	TEXT(""),
 	FConsoleCommandWithArgsDelegate::CreateStatic(HairCardGeneratorEditor_Impl::DbgCmdCacheGroom));
 
+/* FHairCardGeneratorUtils Definition
+ *****************************************************************************/
+
+bool FHairCardGeneratorUtils::BuildGenerationSettings(const bool bQuerySettings, UGroomAsset* NewGroomAsset, FHairGroupsCardsSourceDescription& CardsDesc,
+	TObjectPtr<UHairCardGeneratorPluginSettings>& GenerationSettings, uint8& AllGenFlags, uint8& PipelineGenFlags)
+{
+	// Find the HairCardGenController in Python (or wherever it may be implemented -- could be moved to C++/Blueprints/etc.)
+	TObjectPtr<UHairCardGenControllerBase> CardGenController = FHairCardGenController::GetHairCardGenerator();
+	if ( !CardGenController )
+	{
+		UE_LOG(LogHairCardGenerator, Error, TEXT("No hair card generator controller found, cannot generate hair cards!"));
+		return false;
+	}
+
+	UE_CLOG(CardsDesc.GenerationSettings && !UHairCardGeneratorPluginSettings::IsCompatibleSettings(CardsDesc.GenerationSettings), LogHairCardGenerator, Warning, TEXT("Old hair-card generation settings (from a different generator) will be discarded when you run this generator."));
+
+	GenerationSettings = CardGenController->GetGroomSettings(NewGroomAsset, CardsDesc.LODIndex);
+	if ( !GenerationSettings )
+	{
+		GenerationSettings = HairCardGeneratorEditor_Impl::CreateGroomSettingsLast(NewGroomAsset, CardsDesc);
+	}
+	else
+	{
+		// Need to update all non-generation UI settings in case they've changed on groom asset
+		GenerationSettings->SetGenerateForGroomGroup(CardsDesc.GroupIndex);
+		GenerationSettings->PostResetUpdates();
+	}
+
+	if(bQuerySettings)
+	{
+		if ( !HairCardGenWindow_Utils::PromptUserWithHairCardGenDialog(GenerationSettings) )
+		{
+			return false;
+		}
+	}
+
+	// Check primary settings object merged and all group settings for differences
+	AllGenFlags = GenerationSettings->GetAllPipelineGeneratedDifferences();
+	if ( !GenerationSettings->CheckGenerationFlags(AllGenFlags, EHairCardGenerationPipeline::All) && CardsDesc.ImportedMesh != nullptr )
+	{
+		UE_LOG(LogHairCardGenerator, Display, TEXT("No settings changes, skipping hair card run"));
+		HairCardGeneratorEditor_Impl::CreateFinishNotification(true, LOCTEXT("HairCardGen.UpToDate", "Hair card assets are up to date"));
+		// FIXME: Currently returning true clears textures unless a reimport is run
+		return false;
+	}
+
+	// Force reimport if we are going to run any pipeline step
+	AllGenFlags |= (uint8)EHairCardGenerationPipeline::ImportUpdate;
+
+	// Check primary settings object (without group settings) used to force reruns for e.g. random seed changes
+	PipelineGenFlags = GenerationSettings->GetPipelineGeneratedDifferences();
+
+	// Write the settings that we will attempt to output to card desc (won't be stored on failure see FGroomRenderingDetails::OnGenerateCardDataUsingPlugin for details)
+	CardsDesc.GenerationSettings = GenerationSettings;
+
+	// Need to set the CardsDesc texture layout appropriately in case we succeed
+	CardsDesc.Textures.SetLayout(GenerationSettings->ChannelLayout);
+
+	// TODO: On error will force full regen (is this the best we can do?)
+	GenerationSettings->ClearPipelineGenerated();
+
+	// Output generation settings to json file in Metadata folder
+	GenerationSettings->WriteGenerationSettings();
+	
+	return true;
+}
+
+
+bool FHairCardGeneratorUtils::LoadGenerationSettings(const TObjectPtr<UHairCardGeneratorPluginSettings>& GenerationSettings)
+{
+	// Find the HairCardGenController in Python (or wherever it may be implemented -- could be moved to C++/Blueprints/etc.)
+	TObjectPtr<UHairCardGenControllerBase> CardGenController = FHairCardGenController::GetHairCardGenerator();
+	if ( !CardGenController )
+	{
+		UE_LOG(LogHairCardGenerator, Error, TEXT("No hair card generator controller found, cannot generate hair cards!"));
+		return false;
+	}
+	
+	if ( !CardGenController->LoadSettings(GenerationSettings) )
+	{
+		UE_LOG(LogHairCardGenerator, Error, TEXT("Failed to load hair card generation settings. See log for details."));
+		HairCardGeneratorEditor_Impl::CreateFinishNotification(false, LOCTEXT("HairCardGen.FailedSettingsLoad", "Failed to load hair card generation settings. See log for details."));
+		return false;
+	}
+	return true;
+}
+
+bool FHairCardGeneratorUtils::GenerateCardsClumps(const TObjectPtr<const UHairCardGeneratorPluginSettings>& GenerationSettings, const int32 FilterIndex, const uint8 GenFlags, TArray<int32>& StrandsClumps, int32& NumClumps)
+{
+	TObjectPtr<UHairCardGenControllerBase> CardGenController = FHairCardGenController::GetHairCardGenerator();
+	if ( !CardGenController )
+	{
+		return false;
+	}
+
+	if ( UHairCardGeneratorPluginSettings::CheckGenerationFlags(GenFlags, EHairCardGenerationPipeline::StrandClustering) && !GenerationSettings->bReduceCardsFromPreviousLOD )
+	{
+		StrandsClumps = CardGenController->GenerateClumps(FilterIndex);
+		if (StrandsClumps.IsEmpty())
+		{
+			UE_LOG(LogHairCardGenerator, Error, TEXT("Strand clustering failed for settings group %d. Check the log for details."), FilterIndex);
+			return false;
+		}
+		else  
+		{
+			NumClumps = StrandsClumps.Last();
+		}
+	}
+	return true;
+}
+
+bool FHairCardGeneratorUtils::GenerateCardsGeometry(const TObjectPtr<const UHairCardGeneratorPluginSettings>& GenerationSettings, const int32 FilterIndex, const uint8 GenFlags, TArray<TArray<FVector3f>>& ClumpsGeometry)
+{
+	TObjectPtr<UHairCardGenControllerBase> CardGenController = FHairCardGenController::GetHairCardGenerator();
+	if ( !CardGenController )
+	{
+		return false;
+	}
+	if ( UHairCardGeneratorPluginSettings::CheckGenerationFlags(GenFlags, EHairCardGenerationPipeline::GeometryGeneration) )
+	{
+		if ( !HairCardGeneratorEditor_Impl::GenerateCardsGeometry(GenerationSettings, FilterIndex,ClumpsGeometry) )
+		{
+			UE_LOG(LogHairCardGenerator, Error, TEXT("Cards geometry failed for settings group %d. Check the log for details."), FilterIndex)
+			return false;
+		}
+	}
+	return true;
+}
+
+bool FHairCardGeneratorUtils::GenerateCardsTextures(const TObjectPtr<const UHairCardGeneratorPluginSettings>& GenerationSettings, const int32 FilterIndex, const uint8 GenFlags,
+	TArray<int32>& CardsTextures, int32& NumTextures, TArray<float>& VertexUVs)
+{
+	TObjectPtr<UHairCardGenControllerBase> CardGenController = FHairCardGenController::GetHairCardGenerator();
+	if ( !CardGenController )
+	{
+		return false;
+	}
+	if ( UHairCardGeneratorPluginSettings::CheckGenerationFlags(GenFlags, EHairCardGenerationPipeline::TextureClustering) && !GenerationSettings->bReduceCardsFromPreviousLOD )
+	{
+		CardsTextures = CardGenController->ClusterTextures(FilterIndex);
+		if (CardsTextures.IsEmpty())
+		{
+			UE_LOG(LogHairCardGenerator, Error, TEXT("Texture clustering failed for settings group %d. Check the log for details."), FilterIndex)
+			return false;
+		}
+		else
+		{
+			NumTextures = CardsTextures.Last();
+		}
+	}
+	if ( UHairCardGeneratorPluginSettings::CheckGenerationFlags(GenFlags, EHairCardGenerationPipeline::TextureLayout) && !GenerationSettings->bReduceCardsFromPreviousLOD )
+	{
+		VertexUVs = CardGenController->GenerateTextureLayout();
+		if ( VertexUVs.IsEmpty() )
+		{
+			UE_LOG(LogHairCardGenerator, Error, TEXT("Texture layout failed. See log for details."));
+			return false;
+		}
+	}
+	
+	if ( UHairCardGeneratorPluginSettings::CheckGenerationFlags(GenFlags, EHairCardGenerationPipeline::TextureRendering) && !GenerationSettings->bReduceCardsFromPreviousLOD )
+	{
+		if ( !CardGenController->GenerateTextureAtlases() )
+		{
+			UE_LOG(LogHairCardGenerator, Error, TEXT("Texture atlas creation failed. See log for details."));
+			return false;
+		}
+	}
+	return true;
+}
+
+bool FHairCardGeneratorUtils::RunCardsGeneration(const TObjectPtr<UHairCardGeneratorPluginSettings>& GenerationSettings, const uint8 PipelineFlags, const FGeneratorFunction& PipelineFunction, const bool bCheckFlags)
+{
+	// Run card group generation in loop
+	bool bAllGenerated = true;
+	for ( int32 FilterIndex = 0; FilterIndex < GenerationSettings->GetFilterGroupSettings().Num(); ++FilterIndex )
+	{
+		const uint8 GroupGenFlags = bCheckFlags ? (PipelineFlags | GenerationSettings->GetPipelineFilterGroupDifferences(FilterIndex)) : static_cast<uint8>(EHairCardGenerationPipeline::All);
+
+		// Skip if no group update flags required for current group
+		if ( !GenerationSettings->CheckGenerationFlags(GroupGenFlags, EHairCardGenerationPipeline::GroupUpdate) )
+		{
+			UE_LOG(LogHairCardGenerator, Display, TEXT("Setttings group %d already generated. Skipping."), FilterIndex);
+			continue;
+		}
+
+		// TODO: On errors this will force full regeneration (because we clear the group gen object)
+		GenerationSettings->ClearPipelineGeneratedFilterGroup(FilterIndex);
+		bAllGenerated &= PipelineFunction(GenerationSettings, FilterIndex, GroupGenFlags);
+
+		// TODO: It might be useful to handle user cancellation by breaking vs. errors by continuing here
+		if ( !bAllGenerated )
+		{
+			break;
+		}
+
+		GenerationSettings->WritePipelineGeneratedFilterGroup(FilterIndex);
+	}
+
+	if ( !bAllGenerated )
+	{
+		UE_LOG(LogHairCardGenerator, Error, TEXT("Failed to generate cards for all groups. See log for details."));
+		HairCardGeneratorEditor_Impl::CreateFinishNotification(false, LOCTEXT("HairCardGen.FailedGeneration", "Failed to generate cards for all groups. See log for details."));
+		return false;
+	}
+	return true;
+}
+
+static bool GenerateHairCardGenGroomData(const int PointsPerStrand, const FHairDescription& InGroomData, FHairCardGen_GroomData& Out, const FHairCardGeneratorUtils::FBuilderFunction& BuilderFunction)
+{
+	const int NumStrands = InGroomData.GetNumStrands();
+	Out.Strands.Reserve(NumStrands);
+
+	TGroomAttributesConstRef<float> GroomHairWidth = InGroomData.GroomAttributes().GetAttributesRef<float>(HairAttribute::Groom::Width);
+	const FGroomID GroomID = FGroomID(0);
+	const float Groom_DefaultWidth = GroomHairWidth.IsValid() ? GroomHairWidth[GroomID] : 0.f;
+
+	FStrandAttributesRange TraverseStrands(InGroomData, FHairStrandAttribsConstIterator::DefaultIgnore);
+	if ( !TraverseStrands.IsValid() )
+		return false;
+
+	UE_CLOG(!TraverseStrands.First().HasGroupID(), LogHairCardGenerator, Warning, TEXT("Groom description has no physics group info, may not generate correctly"));
+
+	const FHairStrandAttribsRefProxy& FirstStrand = TraverseStrands.First();
+	UE_CLOG((!FirstStrand.HasCurveType() || !FirstStrand.HasBasisType() || FirstStrand.HasKnots()), LogHairCardGenerator, Warning, TEXT("Invalid strand curve specification, treating as linear"));
+
+	EGroomCurveType CurveType = FirstStrand.GetCurveType();
+	EGroomBasisType BasisType = FirstStrand.GetBasisType();
+
+	if (CurveType == EGroomCurveType::VariableOrder)
+	{
+		UE_LOG(LogHairCardGenerator, Warning, TEXT("VariableOrder curves not supported, treating as linear"));
+
+		CurveType = EGroomCurveType::Linear;
+		BasisType = EGroomBasisType::NoBasis;
+	}
+
+	Out.CurveType = StaticEnum<EGroomCurveType>()->GetNameStringByValue((int64)CurveType);
+	Out.BasisType = StaticEnum<EGroomBasisType>()->GetNameStringByValue((int64)BasisType);
+
+	Out.VertexPositions.Reserve(NumStrands * PointsPerStrand * 3);
+	Out.VertexWidths.Reserve(NumStrands * PointsPerStrand);
+
+	TArray<TArray<FVector>> StrandsPositions;
+	BuilderFunction(StrandsPositions);
+
+	int32 StrandIndex = 0;
+	for ( const FHairStrandAttribsRefProxy& StrandAttribs : TraverseStrands )
+	{
+		FHairCardGen_StrandData& OutStrandData = Out.Strands.AddDefaulted_GetRef();
+		OutStrandData.GroupID = StrandAttribs.HasGroupID() ? StrandAttribs.GetGroupID() : -1;
+
+		const float Strand_DefaultWidth = StrandAttribs.HasStrandWidth() ? StrandAttribs.GetStrandWidth() : Groom_DefaultWidth;
+
+		TArray<FVector> VertexPositions;
+		VertexPositions.SetNum(StrandAttribs.NumVerts());
+		TArray<float> VertexWidths;
+		VertexWidths.SetNum(StrandAttribs.NumVerts());
+
+		const bool bIsVertexValid = StrandsPositions.IsValidIndex(StrandIndex) && (StrandAttribs.NumVerts() <= StrandsPositions[StrandIndex].Num());
+		for (int StrandVertIndex = 0; StrandVertIndex < StrandAttribs.NumVerts(); ++StrandVertIndex)
+		{
+			VertexPositions[StrandVertIndex] = bIsVertexValid ? StrandsPositions[StrandIndex][StrandVertIndex] : FVector(StrandAttribs.GetVertexPosition(StrandVertIndex));
+			VertexWidths[StrandVertIndex] = StrandAttribs.HasVertexWidths() ? StrandAttribs.GetVertexWidth(StrandVertIndex) : Strand_DefaultWidth;
+		}
+
+		FHairCardGenStrandNSCInterpolator::StrandInterpolationResult InterpolatedStrand = FHairCardGenStrandNSCInterpolator(VertexPositions, VertexWidths).GetInterpolatedStrand(PointsPerStrand);
+
+		Out.VertexPositions.Append(InterpolatedStrand.Positions);
+		Out.VertexWidths.Append(InterpolatedStrand.Widths);
+		++StrandIndex;
+	}
+
+	Out.Strands.Shrink();
+	Out.VertexPositions.Shrink();
+	Out.VertexWidths.Shrink();
+
+	Out.HairlineGroupID = INDEX_NONE;
+
+	return true;
+}
+
+bool FHairCardGeneratorUtils::LoadGroomStrands(const UGroomAsset* GroomAsset, const FBuilderFunction& BuilderFunction)
+{
+	TObjectPtr<UHairCardGenControllerBase> CardGenController = FHairCardGenController::GetHairCardGenerator();
+	if ( !CardGenController )
+	{
+		return false;
+	}
+	FHairCardGen_GroomData CardGenData;
+	GenerateHairCardGenGroomData(CardGenController->GetPointsPerCurve(), GroomAsset->GetHairDescription(), CardGenData, BuilderFunction);
+
+	FString CachedGroomsPath = FPaths::ProjectIntermediateDir() / TEXT("GroomHairCardGen") / TEXT("CachedGrooms");
+	return CardGenController->LoadGroomData(CardGenData, GroomAsset->GetName(), CachedGroomsPath);
+}
+
+bool FHairCardGeneratorUtils::BuildCardsAssets(UGroomAsset* NewGroomAsset, FHairGroupsCardsSourceDescription& CardsDesc,
+	const TObjectPtr<UHairCardGeneratorPluginSettings>& GenerationSettings, const uint8 GenFlags)
+{
+	
+	TObjectPtr<UHairCardGenControllerBase> CardGenController = FHairCardGenController::GetHairCardGenerator();
+	if ( !CardGenController )
+	{
+		return false;
+	}
+
+	// Import Textures (NOTE: Mesh import no longer supported)
+	if ( GenerationSettings->CheckGenerationFlags(GenFlags, EHairCardGenerationPipeline::ImportTextures) )
+	{
+		FHairCardImportSettings ImportSettings;
+		HairCardGeneratorEditor_Impl::GetTextureImportNames(GenerationSettings, ImportSettings);
+
+		const FString TexturePath = GenerationSettings->GetTextureImportPath();
+		const FString TextureDest = GenerationSettings->GetTextureContentPath();
+		UE_LOG(LogHairCardGenerator, Display, TEXT("Importing generated cards from: %s"), *TexturePath);
+		TArray<UObject*> ImportedAssets = HairCardGeneratorEditor_Impl::ImportHairCardTextures(TexturePath, TextureDest, ImportSettings);
+		if ( ImportedAssets.IsEmpty() )
+		{
+			UE_LOG(LogHairCardGenerator, Error, TEXT("Failed to import textures. Check '%s' for the files."), *TexturePath);
+			HairCardGeneratorEditor_Impl::CreateFinishNotification(false, LOCTEXT("HairCardGen.FailedTextureImport", "Failed to import textures. See log for details."));
+			return false;
+		}
+
+		const int32 NumAssigned = HairCardGeneratorEditor_Impl::AssignImportedTextures(ImportedAssets, ImportSettings, CardsDesc);
+		if ( NumAssigned == 0 )
+		{
+			UE_LOG(LogHairCardGenerator, Error, TEXT("Unable to assign any of the imported assets to the groom. Verify texture file name format."));
+			HairCardGeneratorEditor_Impl::CreateFinishNotification(false, LOCTEXT("HairCardGen.FailedAssignImports", "Unable to assign any of the imported assets to the groom. See log for details."));
+			return false;
+		}
+	}
+	
+	// Use num groups and out group id to determine mesh name
+	const int32 NumPhysGroups = NewGroomAsset->GetHairDescriptionGroups().HairGroups.Num();
+
+	// Generate static mesh from cards and add to destination package
+	if ( GenerationSettings->CheckGenerationFlags(GenFlags, EHairCardGenerationPipeline::GenerateMesh) )
+	{
+		FScopedSlowTask SlowTask(0, LOCTEXT("GeneratingHairCards.BuildMesh", "Building mesh..."));
+		SlowTask.MakeDialog(/*bShowCancelButton =*/true);
+
+		// Don't bother adding _PG# at end of mesh name if there's only one physics group
+		int OutGroupId = (NumPhysGroups > 1) ? GenerationSettings->GetGenerateForGroomGroup() : -1;
+		TObjectPtr<UStaticMesh> NewMesh = HairCardGeneratorEditor_Impl::CreateMeshInPackage(GenerationSettings->DestinationPath.Path, GenerationSettings->BaseFilename, OutGroupId);
+		if ( !NewMesh )
+		{
+			
+			UE_LOG(LogHairCardGenerator, Error, TEXT("Unable to create mesh in destination package."));
+			HairCardGeneratorEditor_Impl::CreateFinishNotification(false, LOCTEXT("HairCardGen.FailedObjectCreate", "Unable to create mesh in destination package. See log for details."));
+			return false;
+		}
+
+		if ( !CardGenController->GenerateMesh(NewMesh) )
+		{
+			UE_LOG(LogHairCardGenerator, Error, TEXT("Mesh generation failed. See log for details."));
+			HairCardGeneratorEditor_Impl::CreateFinishNotification(false, LOCTEXT("HairCardGen.FailedMeshGen", "Mesh generation failed. See log for details."));
+			return false;
+		}
+
+		CardsDesc.ImportedMesh = NewMesh;
+		CardsDesc.UpdateMeshKey();
+
+		NewMesh->PostEditChange();
+		NewMesh->GetPackage()->MarkPackageDirty();
+	}
+
+	GenerationSettings->WritePipelineGenerated();
+	HairCardGeneratorEditor_Impl::CreateFinishNotification(true, LOCTEXT("HairCardGen.GenSuccess", "Hair cards generated successfully."));
+	return true;
+}
+
 /* HairCardGeneratorEditor_Impl Definition
  *****************************************************************************/
@@ -193,7 +564,7 @@
 	}
 }
 
-static bool HairCardGeneratorEditor_Impl::LoadGroomData(UGroomAsset* NewGroomAsset, bool SaveCached = false)
+static bool HairCardGeneratorEditor_Impl::LoadGroomData(UGroomAsset* NewGroomAsset, const bool SaveCached)
 {
 	if (HairCardGeneratorEditor_Impl::LastGroom == NewGroomAsset && !SaveCached)
 	{
@@ -213,7 +584,7 @@
 		SlowTask.MakeDialog(/*bShowCancelButton =*/true);
 
 		FString CachedGroomsPath = FPaths::ProjectIntermediateDir() / TEXT("GroomHairCardGen") / TEXT("CachedGrooms");
-		if (!CardGenController->LoadGroomData(RestructuredData, NewGroomAsset->GetName(), CachedGroomsPath, SaveCached))
+		if (!CardGenController->LoadGroomData(RestructuredData, NewGroomAsset->GetName(), CachedGroomsPath, SaveCached)) 
 		{
 			return false;
 		}
@@ -227,73 +598,10 @@
 
 static bool HairCardGeneratorEditor_Impl::BuildHairCardGenGroomData(const int PointsPerStrand, const FHairDescription& InGroomData, FHairCardGen_GroomData& Out)
 {
-	const int NumStrands = InGroomData.GetNumStrands();
-	Out.Strands.Reserve(NumStrands);
-
-	TGroomAttributesConstRef<float> GroomHairWidth = InGroomData.GroomAttributes().GetAttributesRef<float>(HairAttribute::Groom::Width);
-	const FGroomID GroomID = FGroomID(0);
-	const float Groom_DefaultWidth = GroomHairWidth.IsValid() ? GroomHairWidth[GroomID] : 0.f;
-
-	FStrandAttributesRange TraverseStrands(InGroomData, FHairStrandAttribsConstIterator::DefaultIgnore);
-	if ( !TraverseStrands.IsValid() )
-		return false;
-
-	UE_CLOG(!TraverseStrands.First().HasGroupID(), LogHairCardGenerator, Warning, TEXT("Groom description has no physics group info, may not generate correctly"));
-
-	const FHairStrandAttribsRefProxy& FirstStrand = TraverseStrands.First();
-	UE_CLOG((!FirstStrand.HasCurveType() || !FirstStrand.HasBasisType() || FirstStrand.HasKnots()), LogHairCardGenerator, Warning, TEXT("Invalid strand curve specification, treating as linear"));
-
-	EGroomCurveType CurveType = FirstStrand.GetCurveType();
-	EGroomBasisType BasisType = FirstStrand.GetBasisType();
-
-	if (CurveType == EGroomCurveType::VariableOrder)
-	{
-		UE_LOG(LogHairCardGenerator, Warning, TEXT("VariableOrder curves not supported, treating as linear"));
-
-		CurveType = EGroomCurveType::Linear;
-		BasisType = EGroomBasisType::NoBasis;
-	}
-
-	Out.CurveType = StaticEnum<EGroomCurveType>()->GetNameStringByValue((int64)CurveType);
-	Out.BasisType = StaticEnum<EGroomBasisType>()->GetNameStringByValue((int64)BasisType);
-
-	Out.VertexPositions.Reserve(NumStrands * PointsPerStrand * 3);
-	Out.VertexWidths.Reserve(NumStrands * PointsPerStrand);
-
-	for ( const FHairStrandAttribsRefProxy& StrandAttribs : TraverseStrands )
-	{
-		FHairCardGen_StrandData& OutStrandData = Out.Strands.AddDefaulted_GetRef();
-		OutStrandData.GroupID = StrandAttribs.HasGroupID() ? StrandAttribs.GetGroupID() : -1;
-
-		const float Strand_DefaultWidth = StrandAttribs.HasStrandWidth() ? StrandAttribs.GetStrandWidth() : Groom_DefaultWidth;
-
-		TArray<FVector> VertexPositions;
-		VertexPositions.SetNum(StrandAttribs.NumVerts());
-		TArray<float> VertexWidths;
-		VertexWidths.SetNum(StrandAttribs.NumVerts());
-
-		for (int StrandVertIndex = 0; StrandVertIndex < StrandAttribs.NumVerts(); ++StrandVertIndex)
-		{
-			VertexPositions[StrandVertIndex] = FVector(StrandAttribs.GetVertexPosition(StrandVertIndex));
-			VertexWidths[StrandVertIndex] = StrandAttribs.HasVertexWidths() ? StrandAttribs.GetVertexWidth(StrandVertIndex) : Strand_DefaultWidth;
-		}
-
-		FHairCardGenStrandNSCInterpolator::StrandInterpolationResult InterpolatedStrand = FHairCardGenStrandNSCInterpolator(VertexPositions, VertexWidths).GetInterpolatedStrand(PointsPerStrand);
-
-		Out.VertexPositions.Append(InterpolatedStrand.Positions);
-		Out.VertexWidths.Append(InterpolatedStrand.Widths);
-	}
-
-	Out.Strands.Shrink();
-	Out.VertexPositions.Shrink();
-	Out.VertexWidths.Shrink();
-
-	Out.HairlineGroupID = INDEX_NONE;
-
-	return true;
+	return GenerateHairCardGenGroomData(PointsPerStrand, InGroomData, Out, [](TArray<TArray<FVector>>&){});
 }
 
-static bool HairCardGeneratorEditor_Impl::GenerateCardsGeometry(TObjectPtr<const UHairCardGeneratorPluginSettings> Settings, int index)
+static bool HairCardGeneratorEditor_Impl::GenerateCardsGeometry(TObjectPtr<const UHairCardGeneratorPluginSettings> Settings, int index, TArray<TArray<FVector3f>>& ClumpsGeometry)
 {
 	TObjectPtr<UHairCardGenControllerBase> CardGenController = FHairCardGenController::GetHairCardGenerator();
 	if ( !CardGenController )
@@ -402,55 +710,67 @@
 	}
 
 	SlowTask.EnterProgressFrame(9);
-	if (!CardGenController->GenerateCardsGeometry())
+	TArray<float> ClumpsVertices = CardGenController->GenerateCardsGeometry();
+	if (ClumpsVertices.IsEmpty())
 	{
 		UE_LOG(LogHairCardGenerator, Error, TEXT("Card geometry generation failed for settings group %d. Check the log for details."), index);
 		return false;
 	}
+	
+	// Fill the vertices coming from the optimizer
+	ClumpsGeometry.Reset();
+	ClumpsGeometry.AddDefaulted();
+	
+	int32 VertexIndex = 0;
+	while(VertexIndex < ClumpsVertices.Num())
+	{
+		if(ClumpsVertices[VertexIndex] == -1.0)
+		{
+			if(VertexIndex != (ClumpsVertices.Num()-1))
+			{
+				ClumpsGeometry.AddDefaulted();
+			}
+			VertexIndex += 1;
+		}
+		else
+		{
+			ClumpsGeometry.Last().Add(FVector3f(ClumpsVertices[VertexIndex], ClumpsVertices[VertexIndex+1], ClumpsVertices[VertexIndex+2]));
+			VertexIndex += 3;
+		}
+	}
 
 	return true;
 }
 
-static bool HairCardGeneratorEditor_Impl::GenerateCardsForCardGroup(TObjectPtr<const UHairCardGeneratorPluginSettings> Settings, int index, uint8 GenFlags)
+static bool HairCardGeneratorEditor_Impl::GenerateCardsForCardGroup(const TObjectPtr<const UHairCardGeneratorPluginSettings> GenerationSettings, const int32 FilterIndex, const uint8 GenFlags)
 {
-	TObjectPtr<UHairCardGenControllerBase> CardGenController = FHairCardGenController::GetHairCardGenerator();
-	if ( !CardGenController )
-	{
-		return false;
-	}
-
-	FScopedSlowTask SlowTask(10, FText::Format(LOCTEXT("GeneratingHairCards.GenerateGroup", "Generating Hair Cards for settings group {0}..."), index));
+	FScopedSlowTask SlowTask(10, FText::Format(LOCTEXT("GeneratingHairCards.GenerateGroup", "Generating Hair Cards for settings group {0}..."), FilterIndex));
 	SlowTask.MakeDialog(/*bShowCancelButton =*/true);
 
 	SlowTask.EnterProgressFrame(1);
-	if ( UHairCardGeneratorPluginSettings::CheckGenerationFlags(GenFlags, EHairCardGenerationPipeline::StrandClustering) && !Settings->bReduceCardsFromPreviousLOD )
+	TArray<int32> StrandsClumps;
+	int32 NumClumps = 0;
+	if(!FHairCardGeneratorUtils::GenerateCardsClumps(GenerationSettings, FilterIndex, GenFlags, StrandsClumps, NumClumps))
 	{
-		if ( !CardGenController->GenerateClumps(index) )
-		{
-			UE_LOG(LogHairCardGenerator, Error, TEXT("Strand clustering failed for settings group %d. Check the log for details."), index);
-			SlowTask.EnterProgressFrame(9);
-			return false;
-		}
+		SlowTask.EnterProgressFrame(9);
+		return false;
 	}
 
+	TArray<TArray<FVector3f>> ClumpsGeometry;
 	SlowTask.EnterProgressFrame(5);
-	if ( UHairCardGeneratorPluginSettings::CheckGenerationFlags(GenFlags, EHairCardGenerationPipeline::GeometryGeneration) )
+	if(!FHairCardGeneratorUtils::GenerateCardsGeometry(GenerationSettings, FilterIndex, GenFlags, ClumpsGeometry))
 	{
-		if ( !HairCardGeneratorEditor_Impl::GenerateCardsGeometry(Settings, index) )
-		{
-			SlowTask.EnterProgressFrame(4);
-			return false;
-		}
+		SlowTask.EnterProgressFrame(4);
+		return false;
 	}
 
 	SlowTask.EnterProgressFrame(4);
-	if ( UHairCardGeneratorPluginSettings::CheckGenerationFlags(GenFlags, EHairCardGenerationPipeline::TextureClustering) && !Settings->bReduceCardsFromPreviousLOD )
+	TArray<int32> CardsTextures;
+	int32 NumTextures= 0;
+	TArray<float> VertexUVs;
+	if(!FHairCardGeneratorUtils::GenerateCardsTextures(GenerationSettings, FilterIndex, GenFlags, CardsTextures, NumTextures, VertexUVs))
 	{
-		if (!CardGenController->ClusterTextures(index))
-		{
-			UE_LOG(LogHairCardGenerator, Error, TEXT("Texture clustering failed for settings group %d. Check the log for details."), index)
-			return false;
-		}
+		return false;
 	}
 
 	return true;
@@ -584,20 +904,6 @@
 /* FHairCardGeneratorEditorModule
  *****************************************************************************/
 
-class FHairCardGeneratorEditorModule : public IHairCardGeneratorEditor
-{
-public:
-	//~ Begin IModuleInterface API
-	virtual void StartupModule() override;
-	virtual void ShutdownModule() override;
-	//~ End IModuleInterface API
-
-	//~ Begin IHairCardGenerator API
-	virtual bool GenerateHairCardsForLOD(UGroomAsset* Groom, FHairGroupsCardsSourceDescription& CardsDesc) override;
-	virtual bool IsCompatibleSettings(UHairCardGenerationSettings* OldSettings) override;
-	//~ End IHairCardGenerator API
-};
-
 void FHairCardGeneratorEditorModule::StartupModule()
 {
 	LLM_SCOPE_BYTAG(HairCardGenerator)
@@ -620,6 +926,30 @@
 
 bool FHairCardGeneratorEditorModule::GenerateHairCardsForLOD(UGroomAsset* NewGroomAsset, FHairGroupsCardsSourceDescription& CardsDesc)
 {
+	TObjectPtr<UHairCardGeneratorPluginSettings> GenerationSettings = nullptr;
+	uint8 AllGenFlags = 0, PipelineGenFlags = 0;
+	if(!FHairCardGeneratorUtils::BuildGenerationSettings(true, NewGroomAsset, CardsDesc, GenerationSettings, AllGenFlags, PipelineGenFlags))
+	{
+		return false;
+	}
+
+	if (!HairCardGeneratorEditor_Impl::LoadGroomData(NewGroomAsset))
+	{
+		UE_LOG(LogHairCardGenerator, Error, TEXT("Failed to load groom asset. See log for details."));
+		HairCardGeneratorEditor_Impl::CreateFinishNotification(false, LOCTEXT("HairCardGen.FailedGroomLoad", "Failed to load groom asset. See log for details."));
+		return false;
+	}
+	
+	if(GenerationSettings)
+	{
+		GenerationSettings->UpdateStrandFilterAssignment();
+	}
+
+	if(!FHairCardGeneratorUtils::LoadGenerationSettings(GenerationSettings))
+	{
+		return false;
+	}
+	
 	// Find the HairCardGenController in Python (or wherever it may be implemented -- could be moved to C++/Blueprints/etc.)
 	TObjectPtr<UHairCardGenControllerBase> CardGenController = FHairCardGenController::GetHairCardGenerator();
 	if ( !CardGenController )
@@ -628,188 +958,19 @@
 		return false;
 	}
 
-	UE_CLOG(CardsDesc.GenerationSettings && !IsCompatibleSettings(CardsDesc.GenerationSettings), LogHairCardGenerator, Warning, TEXT("Old hair-card generation settings (from a different generator) will be discarded when you run this generator."));
-
-	TObjectPtr<UHairCardGeneratorPluginSettings>& GenerationSettings = CardGenController->GetGroomSettings(NewGroomAsset, CardsDesc.LODIndex);
-	if ( !GenerationSettings )
-	{
-		GenerationSettings = HairCardGeneratorEditor_Impl::CreateGroomSettingsLast(NewGroomAsset, CardsDesc);
-	}
-	else
-	{
-		// Need to update all non-generation UI settings in case they've changed on groom asset
-		GenerationSettings->SetGenerateForGroomGroup(CardsDesc.GroupIndex);
-		GenerationSettings->PostResetUpdates();
-	}
-
-	bool bSuccess = false;
-	if ( !HairCardGenWindow_Utils::PromptUserWithHairCardGenDialog(GenerationSettings) )
+	if(!FHairCardGeneratorUtils::RunCardsGeneration(GenerationSettings, PipelineGenFlags,
+		[](const TObjectPtr<const UHairCardGeneratorPluginSettings>& GenerationSettings, const int32 FilterIndex, const uint8 GenFlags)
+		{
+			return HairCardGeneratorEditor_Impl::GenerateCardsForCardGroup(GenerationSettings, FilterIndex, GenFlags);
+		}))
 	{
 		return false;
 	}
 
-	// Use num groups and out group id to determine mesh name
-	int NumPhysGroups = NewGroomAsset->GetHairDescriptionGroups().HairGroups.Num();
-
-	// Check primary settings object merged and all group settings for differences
-	uint8 AllGenFlags = GenerationSettings->GetAllPipelineGeneratedDifferences();
-	if ( !GenerationSettings->CheckGenerationFlags(AllGenFlags, EHairCardGenerationPipeline::All) && CardsDesc.ImportedMesh != nullptr )
+	if(!FHairCardGeneratorUtils::BuildCardsAssets(NewGroomAsset, CardsDesc, GenerationSettings, AllGenFlags))
 	{
-		UE_LOG(LogHairCardGenerator, Display, TEXT("No settings changes, skipping hair card run"));
-		HairCardGeneratorEditor_Impl::CreateFinishNotification(true, LOCTEXT("HairCardGen.UpToDate", "Hair card assets are up to date"));
-		// FIXME: Currently returning true clears textures unless a reimport is run
 		return false;
 	}
-
-	// Force reimport if we are going to run any pipeline step
-	AllGenFlags |= (uint8)EHairCardGenerationPipeline::ImportUpdate;
-
-	// Check primary settings object (without group settings) used to force reruns for e.g. random seed changes
-	uint8 PipelineGenFlags = GenerationSettings->GetPipelineGeneratedDifferences();
-
-	// Write the settings that we will attempt to output to card desc (won't be stored on failure see FGroomRenderingDetails::OnGenerateCardDataUsingPlugin for details)
-	CardsDesc.GenerationSettings = GenerationSettings;
-
-	// Need to set the CardsDesc texture layout appropriately in case we succeed
-	CardsDesc.Textures.SetLayout(GenerationSettings->ChannelLayout);
-
-	// TODO: On error will force full regen (is this the best we can do?)
-	GenerationSettings->ClearPipelineGenerated();
-
-	// Output generation settings to json file in Metadata folder
-	GenerationSettings->WriteGenerationSettings();
-
-	if (!HairCardGeneratorEditor_Impl::LoadGroomData(NewGroomAsset))
-	{
-		UE_LOG(LogHairCardGenerator, Error, TEXT("Failed to load groom asset. See log for details."));
-		HairCardGeneratorEditor_Impl::CreateFinishNotification(false, LOCTEXT("HairCardGen.FailedGroomLoad", "Failed to load groom asset. See log for details."));
-		return false;
-	}
-
-	GenerationSettings->UpdateStrandFilterAssignment();
-	if ( !CardGenController->LoadSettings(GenerationSettings) )
-	{
-		UE_LOG(LogHairCardGenerator, Error, TEXT("Failed to load hair card generation settings. See log for details."));
-		HairCardGeneratorEditor_Impl::CreateFinishNotification(false, LOCTEXT("HairCardGen.FailedSettingsLoad", "Failed to load hair card generation settings. See log for details."));
-		return false;
-	}
-
-	// Run card group generation in loop
-	bool bAllGenerated = true;
-	for ( int i=0; i < GenerationSettings->GetFilterGroupSettings().Num(); ++i )
-	{
-		uint8 GroupGenFlags = PipelineGenFlags | GenerationSettings->GetPipelineFilterGroupDifferences(i);
-
-		// Skip if no group update flags required for current group
-		if ( !GenerationSettings->CheckGenerationFlags(GroupGenFlags, EHairCardGenerationPipeline::GroupUpdate) )
-		{
-			UE_LOG(LogHairCardGenerator, Display, TEXT("Setttings group %d already generated. Skipping."), i);
-			continue;
-		}
-
-		// TODO: On errors this will force full regeneration (because we clear the group gen object)
-		GenerationSettings->ClearPipelineGeneratedFilterGroup(i);
-		bAllGenerated &= HairCardGeneratorEditor_Impl::GenerateCardsForCardGroup(GenerationSettings, i, GroupGenFlags);
-
-		// TODO: It might be useful to handle user cancellation by breaking vs. errors by continuing here
-		if ( !bAllGenerated )
-		{
-			break;
-		}
-
-		GenerationSettings->WritePipelineGeneratedFilterGroup(i);
-	}
-
-	if ( !bAllGenerated )
-	{
-		UE_LOG(LogHairCardGenerator, Error, TEXT("Failed to generate cards for all groups. See log for details."));
-		HairCardGeneratorEditor_Impl::CreateFinishNotification(false, LOCTEXT("HairCardGen.FailedGeneration", "Failed to generate cards for all groups. See log for details."));
-		return false;
-	}
-
-	if ( GenerationSettings->CheckGenerationFlags(AllGenFlags, EHairCardGenerationPipeline::TextureLayout) && !GenerationSettings->bReduceCardsFromPreviousLOD )
-	{
-		FScopedSlowTask SlowTask(0, LOCTEXT("HairCardGen.BuildingLayout", "Building texture atlas layout..."));
-		SlowTask.MakeDialog(/*bShowCancelButton =*/true);
-		if ( !CardGenController->GenerateTextureLayout() )
-		{
-			UE_LOG(LogHairCardGenerator, Error, TEXT("Texture layout failed. See log for details."));
-			HairCardGeneratorEditor_Impl::CreateFinishNotification(false, LOCTEXT("HairCardGen.FailedLayout", "Texture layout failed. See log for details."));
-			return false;
-		}
-	}
-
-	if ( GenerationSettings->CheckGenerationFlags(AllGenFlags, EHairCardGenerationPipeline::TextureRendering) && !GenerationSettings->bReduceCardsFromPreviousLOD )
-	{
-		FScopedSlowTask SlowTask(0, LOCTEXT("HairCardGen.BuildingAtlas", "Exporting textures..."));
-		SlowTask.MakeDialog(/*bShowCancelButton =*/true);
-		if ( !CardGenController->GenerateTextureAtlases() )
-		{
-			UE_LOG(LogHairCardGenerator, Error, TEXT("Texture atlas creation failed. See log for details."));
-			HairCardGeneratorEditor_Impl::CreateFinishNotification(false, LOCTEXT("HairCardGen.FailedAtlas", "Texture atlas creation failed. See log for details."));
-			return false;
-		}
-	}
-
-	// Import Textures (NOTE: Mesh import no longer supported)
-	if ( GenerationSettings->CheckGenerationFlags(AllGenFlags, EHairCardGenerationPipeline::ImportTextures) )
-	{
-		FHairCardImportSettings ImportSettings;
-		HairCardGeneratorEditor_Impl::GetTextureImportNames(GenerationSettings, ImportSettings);
-
-		const FString TexturePath = GenerationSettings->GetTextureImportPath();
-		const FString TextureDest = GenerationSettings->GetTextureContentPath();
-		UE_LOG(LogHairCardGenerator, Display, TEXT("Importing generated cards from: %s"), *TexturePath);
-		TArray<UObject*> ImportedAssets = HairCardGeneratorEditor_Impl::ImportHairCardTextures(TexturePath, TextureDest, ImportSettings);
-		if ( ImportedAssets.IsEmpty() )
-		{
-			UE_LOG(LogHairCardGenerator, Error, TEXT("Failed to import textures. Check '%s' for the files."), *TexturePath);
-			HairCardGeneratorEditor_Impl::CreateFinishNotification(false, LOCTEXT("HairCardGen.FailedTextureImport", "Failed to import textures. See log for details."));
-			return false;
-		}
-
-		const int32 NumAssigned = HairCardGeneratorEditor_Impl::AssignImportedTextures(ImportedAssets, ImportSettings, CardsDesc);
-		if ( NumAssigned == 0 )
-		{
-			UE_LOG(LogHairCardGenerator, Error, TEXT("Unable to assign any of the imported assets to the groom. Verify texture file name format."));
-			HairCardGeneratorEditor_Impl::CreateFinishNotification(false, LOCTEXT("HairCardGen.FailedAssignImports", "Unable to assign any of the imported assets to the groom. See log for details."));
-			return false;
-		}
-	}
-
-	// Generate static mesh from cards and add to destination package
-	if ( GenerationSettings->CheckGenerationFlags(AllGenFlags, EHairCardGenerationPipeline::GenerateMesh) )
-	{
-		FScopedSlowTask SlowTask(0, LOCTEXT("GeneratingHairCards.BuildMesh", "Building mesh..."));
-		SlowTask.MakeDialog(/*bShowCancelButton =*/true);
-
-		// Don't bother adding _PG# at end of mesh name if there's only one physics group
-		int OutGroupId = (NumPhysGroups > 1) ? GenerationSettings->GetGenerateForGroomGroup() : -1;
-		TObjectPtr<UStaticMesh> NewMesh = HairCardGeneratorEditor_Impl::CreateMeshInPackage(GenerationSettings->DestinationPath.Path, GenerationSettings->BaseFilename, OutGroupId);
-		if ( !NewMesh )
-		{
-			
-			UE_LOG(LogHairCardGenerator, Error, TEXT("Unable to create mesh in destination package."));
-			HairCardGeneratorEditor_Impl::CreateFinishNotification(false, LOCTEXT("HairCardGen.FailedObjectCreate", "Unable to create mesh in destination package. See log for details."));
-			return false;
-		}
-
-		if ( !CardGenController->GenerateMesh(NewMesh) )
-		{
-			UE_LOG(LogHairCardGenerator, Error, TEXT("Mesh generation failed. See log for details."));
-			HairCardGeneratorEditor_Impl::CreateFinishNotification(false, LOCTEXT("HairCardGen.FailedMeshGen", "Mesh generation failed. See log for details."));
-			return false;
-		}
-
-		CardsDesc.ImportedMesh = NewMesh;
-		CardsDesc.UpdateMeshKey();
-
-		NewMesh->PostEditChange();
-		NewMesh->GetPackage()->MarkPackageDirty();
-	}
-
-	GenerationSettings->WritePipelineGenerated();
-	HairCardGeneratorEditor_Impl::CreateFinishNotification(true, LOCTEXT("HairCardGen.GenSuccess", "Hair cards generated successfully."));
 	return true;
 }
 
diff --git a/Engine/Plugins/Experimental/HairCardGenerator/Source/HairCardGeneratorEditor/Private/HairCardGeneratorLog.h b/Engine/Plugins/Experimental/HairCardGenerator/Source/HairCardGeneratorEditor/Private/HairCardGeneratorLog.h
--- a/Engine/Plugins/Experimental/HairCardGenerator/Source/HairCardGeneratorEditor/Private/HairCardGeneratorLog.h
+++ b/Engine/Plugins/Experimental/HairCardGenerator/Source/HairCardGeneratorEditor/Private/HairCardGeneratorLog.h
@@ -4,4 +4,4 @@
 
 #include "Logging/LogMacros.h"
 
-DECLARE_LOG_CATEGORY_EXTERN(LogHairCardGenerator, Log, All)
\ No newline at end of file
+HAIRCARDGENERATOREDITOR_API DECLARE_LOG_CATEGORY_EXTERN(LogHairCardGenerator, Log, All)
\ No newline at end of file
diff --git a/Engine/Plugins/Experimental/HairCardGenerator/Source/HairCardGeneratorEditor/Private/HairCardGeneratorPluginSettings.h b/Engine/Plugins/Experimental/HairCardGenerator/Source/HairCardGeneratorEditor/Private/HairCardGeneratorPluginSettings.h
--- a/Engine/Plugins/Experimental/HairCardGenerator/Source/HairCardGeneratorEditor/Private/HairCardGeneratorPluginSettings.h
+++ b/Engine/Plugins/Experimental/HairCardGenerator/Source/HairCardGeneratorEditor/Private/HairCardGeneratorPluginSettings.h
@@ -106,7 +106,7 @@ enum class EHairCardGenerationSettingsCategories : uint8
  * hair card generation extension module (see IHairCardGenerator API).
 */
 UCLASS(BlueprintType, config=Editor)
-class UHairCardGeneratorPluginSettings : public UHairCardGenerationSettings
+class HAIRCARDGENERATOREDITOR_API UHairCardGeneratorPluginSettings : public UHairCardGenerationSettings
 {
 	GENERATED_BODY()
 
@@ -234,6 +234,7 @@ public:
 	
 	FString GetGroomName() const;
 
+	UGroomAsset* GetGroomAsset() const {return GroomAsset;}
 	int32 GetLODIndex() const { return LODIndex; }
 	int32 GetGenerateForGroomGroup() const { return GenerateForGroomGroup; }
 	void SetGenerateForGroomGroup(int GroupIndex) { GenerateForGroomGroup = GroupIndex; }
diff --git a/Engine/Plugins/Experimental/HairCardGenerator/Source/HairCardGeneratorEditor/Public/HairCardGeneratorEditorModule.h b/Engine/Plugins/Experimental/HairCardGenerator/Source/HairCardGeneratorEditor/Public/HairCardGeneratorEditorModule.h
new file mode 100644
--- /dev/null
+++ b/Engine/Plugins/Experimental/HairCardGenerator/Source/HairCardGeneratorEditor/Public/HairCardGeneratorEditorModule.h
@@ -0,0 +1,59 @@
+// Copyright Epic Games, Inc. All Rights Reserved.
+
+#pragma once
+#include "IHairCardGeneratorEditor.h"
+#include <functional>
+
+struct FHairGroupsCardsSourceDescription;
+class UGroomAsset;
+class UHairCardGeneratorPluginSettings;
+class UHairCardGenControllerBase;
+
+/** FHairCardGeneratorUtils has static interface functions to be used in other modules */
+struct HAIRCARDGENERATOREDITOR_API FHairCardGeneratorUtils
+{
+	using FGeneratorFunction = std::function<bool(const TObjectPtr<const UHairCardGeneratorPluginSettings>&, const int32, const uint8)>;
+	using FBuilderFunction = std::function<void(TArray<TArray<FVector>>&)>;
+	
+	/** Build the generation settings for a given grom asset and a cards description */
+	static bool BuildGenerationSettings(const bool bQuerySettings, UGroomAsset* GroomAsset, FHairGroupsCardsSourceDescription& CardsDesc,
+		TObjectPtr<UHairCardGeneratorPluginSettings>& GenerationSettings, uint8& GenerationFlags, uint8& PipelineFlags);
+
+	/** Load the generation settings for a given grom asset and a cards description */
+	static bool LoadGenerationSettings(const TObjectPtr<UHairCardGeneratorPluginSettings>& GenerationSettings);
+
+	/** Generate the cards clumps based on the settings given a filter index */
+	static bool GenerateCardsClumps(const TObjectPtr<const UHairCardGeneratorPluginSettings>& GenerationSettings, const int32 FilterIndex, const uint8 GenFlags, TArray<int32>& StrandsClumps, int32& NumClumps);
+
+	/** Generate the cards geometry based on the settings given a filter index */
+	static bool GenerateCardsGeometry(const TObjectPtr<const UHairCardGeneratorPluginSettings>& GenerationSettings, const int32 FilterIndex, const uint8 GenFlags, TArray<TArray<FVector3f>>& ClumpsGeometry);
+
+	/** Generate the cards texture based on the settings given a filter index */
+	static bool GenerateCardsTextures(const TObjectPtr<const UHairCardGeneratorPluginSettings>& GenerationSettings, const int32 FilterIndex, const uint8 GenFlags, TArray<int32>& CardsTextures, int32& NumTextures, TArray<float>& VertexUVs);
+	
+	/** Run the cards generation function */
+	static bool RunCardsGeneration(const TObjectPtr<UHairCardGeneratorPluginSettings>& GenerationSettings, const uint8 PipelineFlags, const FGeneratorFunction& PipelineFunction, const bool bCheckFlags = true);
+
+	/** Build and save the cards assets to disk */
+	static bool BuildCardsAssets(UGroomAsset* NewGroomAsset, FHairGroupsCardsSourceDescription& CardsDesc,
+		const TObjectPtr<UHairCardGeneratorPluginSettings>& GenerationSettings, const uint8 GenFlags);
+
+	/** Load the groom strands */
+	static bool LoadGroomStrands(const UGroomAsset* GroomAsset, const FBuilderFunction& BuilderFunction);
+};
+
+/** FHairCardGeneratorEditorModule  */
+class FHairCardGeneratorEditorModule : public IHairCardGeneratorEditor
+{
+public:
+	//~ Begin IModuleInterface API
+	virtual void StartupModule() override;
+	virtual void ShutdownModule() override;
+	//~ End IModuleInterface API
+
+	//~ Begin IHairCardGenerator API
+	virtual bool GenerateHairCardsForLOD(UGroomAsset* Groom, FHairGroupsCardsSourceDescription& CardsDesc) override;
+	virtual bool IsCompatibleSettings(UHairCardGenerationSettings* OldSettings) override;
+	//~ End IHairCardGenerator API
+};
+
