Fix missing pragma once in cards rendering

diff --git a/Engine/Plugins/Experimental/HairCardGenerator/Source/HairCardGeneratorDataflow/Public/HairCardDataflowRendering.h b/Engine/Plugins/Experimental/HairCardGenerator/Source/HairCardGeneratorDataflow/Public/HairCardDataflowRendering.h
--- a/Engine/Plugins/Experimental/HairCardGenerator/Source/HairCardGeneratorDataflow/Public/HairCardDataflowRendering.h
+++ b/Engine/Plugins/Experimental/HairCardGenerator/Source/HairCardGeneratorDataflow/Public/HairCardDataflowRendering.h
@@ -1,5 +1,7 @@
 // Copyright Epic Games, Inc. All Rights Reserved.
 
+#pragma once
+
 #include "GroomDataflowRendering.h"
 #include "Dataflow/DataflowRenderingFactory.h"
 #include "Dataflow/DataflowRenderingViewMode.h"
