Set groom dataflow nodes as experimental

diff --git a/Engine/Plugins/Experimental/HairCardGenerator/Source/HairCardGeneratorDataflow/Public/BuildCardsSettingsNode.h b/Engine/Plugins/Experimental/HairCardGenerator/Source/HairCardGeneratorDataflow/Public/BuildCardsSettingsNode.h
--- a/Engine/Plugins/Experimental/HairCardGenerator/Source/HairCardGeneratorDataflow/Public/BuildCardsSettingsNode.h
+++ b/Engine/Plugins/Experimental/HairCardGenerator/Source/HairCardGeneratorDataflow/Public/BuildCardsSettingsNode.h
@@ -34,7 +34,7 @@ struct FGroomCardsSettings
 };
 
 /** Build the cards generation settings */
-USTRUCT(meta = (DataflowGroom))
+USTRUCT(meta = (Experimental, DataflowGroom))
 struct FBuildCardsSettingsNode : public FDataflowNode
 {
 	GENERATED_BODY()
diff --git a/Engine/Plugins/Experimental/HairCardGenerator/Source/HairCardGeneratorDataflow/Public/CardsAssetTerminalNode.h b/Engine/Plugins/Experimental/HairCardGenerator/Source/HairCardGeneratorDataflow/Public/CardsAssetTerminalNode.h
--- a/Engine/Plugins/Experimental/HairCardGenerator/Source/HairCardGeneratorDataflow/Public/CardsAssetTerminalNode.h
+++ b/Engine/Plugins/Experimental/HairCardGenerator/Source/HairCardGeneratorDataflow/Public/CardsAssetTerminalNode.h
@@ -11,7 +11,7 @@
 
 #include "CardsAssetTerminalNode.generated.h"
 
-USTRUCT(meta = (DataflowGroom, DataflowTerminal))
+USTRUCT(meta = (Experimental, DataflowGroom, DataflowTerminal))
 struct FCardsAssetTerminalNode : public FDataflowTerminalNode
 {
 	GENERATED_USTRUCT_BODY()
diff --git a/Engine/Plugins/Experimental/HairCardGenerator/Source/HairCardGeneratorDataflow/Public/GenerateCardsClumpsNode.h b/Engine/Plugins/Experimental/HairCardGenerator/Source/HairCardGeneratorDataflow/Public/GenerateCardsClumpsNode.h
--- a/Engine/Plugins/Experimental/HairCardGenerator/Source/HairCardGeneratorDataflow/Public/GenerateCardsClumpsNode.h
+++ b/Engine/Plugins/Experimental/HairCardGenerator/Source/HairCardGeneratorDataflow/Public/GenerateCardsClumpsNode.h
@@ -34,7 +34,7 @@ struct FCardsClumpsSettings
 };
 
 /** Generate the clumps used to build the cards from the strands */
-USTRUCT(meta = (DataflowGroom))
+USTRUCT(meta = (Experimental, DataflowGroom))
 struct FGenerateCardsClumpsNode : public FDataflowNode
 {
 	GENERATED_BODY()
diff --git a/Engine/Plugins/Experimental/HairCardGenerator/Source/HairCardGeneratorDataflow/Public/GenerateCardsGeometryNode.h b/Engine/Plugins/Experimental/HairCardGenerator/Source/HairCardGeneratorDataflow/Public/GenerateCardsGeometryNode.h
--- a/Engine/Plugins/Experimental/HairCardGenerator/Source/HairCardGeneratorDataflow/Public/GenerateCardsGeometryNode.h
+++ b/Engine/Plugins/Experimental/HairCardGenerator/Source/HairCardGeneratorDataflow/Public/GenerateCardsGeometryNode.h
@@ -30,7 +30,7 @@ struct FCardsGeometrySettings
 };
 
 /** Generate the geometry used to build the cards from the strands */
-USTRUCT(meta = (DataflowGroom))
+USTRUCT(meta = (Experimental, DataflowGroom))
 struct FGenerateCardsGeometryNode : public FDataflowNode
 {
 	GENERATED_BODY()
diff --git a/Engine/Plugins/Experimental/HairCardGenerator/Source/HairCardGeneratorDataflow/Public/GenerateCardsTexturesNode.h b/Engine/Plugins/Experimental/HairCardGenerator/Source/HairCardGeneratorDataflow/Public/GenerateCardsTexturesNode.h
--- a/Engine/Plugins/Experimental/HairCardGenerator/Source/HairCardGeneratorDataflow/Public/GenerateCardsTexturesNode.h
+++ b/Engine/Plugins/Experimental/HairCardGenerator/Source/HairCardGeneratorDataflow/Public/GenerateCardsTexturesNode.h
@@ -30,7 +30,7 @@ struct FCardsTextureSettings
 };
 
 /** Generate the clumps used to build the cards from the strands */
-USTRUCT(meta = (DataflowGroom))
+USTRUCT(meta = (Experimental, DataflowGroom))
 struct FGenerateCardsTexturesNode : public FDataflowNode
 {
 	GENERATED_BODY()
diff --git a/Engine/Plugins/Runtime/HairStrands/Source/HairStrandsDataflow/Public/AttachGuidesRootsNode.h b/Engine/Plugins/Runtime/HairStrands/Source/HairStrandsDataflow/Public/AttachGuidesRootsNode.h
--- a/Engine/Plugins/Runtime/HairStrands/Source/HairStrandsDataflow/Public/AttachGuidesRootsNode.h
+++ b/Engine/Plugins/Runtime/HairStrands/Source/HairStrandsDataflow/Public/AttachGuidesRootsNode.h
@@ -10,7 +10,7 @@
 #include "AttachGuidesRootsNode.generated.h"
 
 /** Attach the guides roots by setting their kinematic weights to 1.0f */
-USTRUCT(meta = (DataflowGroom))
+USTRUCT(meta = (Experimental, DataflowGroom))
 struct FAttachGuidesRootsDataflowNode : public FDataflowNode
 {
 	GENERATED_BODY()
diff --git a/Engine/Plugins/Runtime/HairStrands/Source/HairStrandsDataflow/Public/BuildGuidesLODsNode.h b/Engine/Plugins/Runtime/HairStrands/Source/HairStrandsDataflow/Public/BuildGuidesLODsNode.h
--- a/Engine/Plugins/Runtime/HairStrands/Source/HairStrandsDataflow/Public/BuildGuidesLODsNode.h
+++ b/Engine/Plugins/Runtime/HairStrands/Source/HairStrandsDataflow/Public/BuildGuidesLODsNode.h
@@ -12,7 +12,7 @@
 
 
 /** Builds the guides LODs */
-USTRUCT(meta = (DataflowGroom))
+USTRUCT(meta = (Experimental, DataflowGroom))
 struct FBuildGuidesLODsDataflowNode : public FDataflowNode
 {
 	GENERATED_BODY()
diff --git a/Engine/Plugins/Runtime/HairStrands/Source/HairStrandsDataflow/Public/GenerateGuidesCurvesNode.h b/Engine/Plugins/Runtime/HairStrands/Source/HairStrandsDataflow/Public/GenerateGuidesCurvesNode.h
--- a/Engine/Plugins/Runtime/HairStrands/Source/HairStrandsDataflow/Public/GenerateGuidesCurvesNode.h
+++ b/Engine/Plugins/Runtime/HairStrands/Source/HairStrandsDataflow/Public/GenerateGuidesCurvesNode.h
@@ -10,7 +10,7 @@
 #include "GenerateGuidesCurvesNode.generated.h"
 
 /** Build the guides curves from the strands */
-USTRUCT(meta = (DataflowGroom))
+USTRUCT(meta = (Experimental, DataflowGroom))
 struct FGenerateGuidesCurvesDataflowNode : public FDataflowNode
 {
 	GENERATED_BODY()
diff --git a/Engine/Plugins/Runtime/HairStrands/Source/HairStrandsDataflow/Public/GetGroomAssetNode.h b/Engine/Plugins/Runtime/HairStrands/Source/HairStrandsDataflow/Public/GetGroomAssetNode.h
--- a/Engine/Plugins/Runtime/HairStrands/Source/HairStrandsDataflow/Public/GetGroomAssetNode.h
+++ b/Engine/Plugins/Runtime/HairStrands/Source/HairStrandsDataflow/Public/GetGroomAssetNode.h
@@ -12,7 +12,7 @@
 #include "GetGroomAssetNode.generated.h"
 
 /** Get the groom asset guides */
-USTRUCT(meta = (DataflowGroom))
+USTRUCT(meta = (Experimental, DataflowGroom))
 struct FGetGroomAssetDataflowNode : public FDataflowNode
 {
 	GENERATED_BODY()
diff --git a/Engine/Plugins/Runtime/HairStrands/Source/HairStrandsDataflow/Public/GetGuidesAttributesNodes.h b/Engine/Plugins/Runtime/HairStrands/Source/HairStrandsDataflow/Public/GetGuidesAttributesNodes.h
--- a/Engine/Plugins/Runtime/HairStrands/Source/HairStrandsDataflow/Public/GetGuidesAttributesNodes.h
+++ b/Engine/Plugins/Runtime/HairStrands/Source/HairStrandsDataflow/Public/GetGuidesAttributesNodes.h
@@ -10,7 +10,7 @@
 #include "GetGuidesAttributesNodes.generated.h"
 
 /** Get the guides attributes names */
-USTRUCT(meta = (DataflowGroom))
+USTRUCT(meta = (Experimental, DataflowGroom))
 struct FGetGuidesKinematicWeightsDataflowNode : public FDataflowNode
 {
 	GENERATED_BODY()
diff --git a/Engine/Plugins/Runtime/HairStrands/Source/HairStrandsDataflow/Public/GroomAssetTerminalNode.h b/Engine/Plugins/Runtime/HairStrands/Source/HairStrandsDataflow/Public/GroomAssetTerminalNode.h
--- a/Engine/Plugins/Runtime/HairStrands/Source/HairStrandsDataflow/Public/GroomAssetTerminalNode.h
+++ b/Engine/Plugins/Runtime/HairStrands/Source/HairStrandsDataflow/Public/GroomAssetTerminalNode.h
@@ -10,7 +10,7 @@
 
 #include "GroomAssetTerminalNode.generated.h"
 
-USTRUCT(meta = (DataflowGroom, DataflowTerminal))
+USTRUCT(meta = (Experimental, DataflowGroom, DataflowTerminal))
 struct FGroomAssetTerminalDataflowNode : public FDataflowTerminalNode
 {
 	GENERATED_USTRUCT_BODY()
diff --git a/Engine/Plugins/Runtime/HairStrands/Source/HairStrandsDataflow/Public/ResampleGuidesPointsNode.h b/Engine/Plugins/Runtime/HairStrands/Source/HairStrandsDataflow/Public/ResampleGuidesPointsNode.h
--- a/Engine/Plugins/Runtime/HairStrands/Source/HairStrandsDataflow/Public/ResampleGuidesPointsNode.h
+++ b/Engine/Plugins/Runtime/HairStrands/Source/HairStrandsDataflow/Public/ResampleGuidesPointsNode.h
@@ -34,7 +34,7 @@ enum class EGroomNumPoints : uint8
 };
 
 /** Resample the groom guides with a fixed number of points (in physics strands size) */
-USTRUCT(meta = (DataflowGroom))
+USTRUCT(meta = (Experimental, DataflowGroom))
 struct FResampleGuidesPointsDataflowNode : public FDataflowNode
 {
 	GENERATED_BODY()
diff --git a/Engine/Plugins/Runtime/HairStrands/Source/HairStrandsDataflow/Public/SmoothGuidesCurvesNode.h b/Engine/Plugins/Runtime/HairStrands/Source/HairStrandsDataflow/Public/SmoothGuidesCurvesNode.h
--- a/Engine/Plugins/Runtime/HairStrands/Source/HairStrandsDataflow/Public/SmoothGuidesCurvesNode.h
+++ b/Engine/Plugins/Runtime/HairStrands/Source/HairStrandsDataflow/Public/SmoothGuidesCurvesNode.h
@@ -12,7 +12,7 @@
 
 
 /** Smooth the guides for more stable simulation */
-USTRUCT(meta = (DataflowGroom))
+USTRUCT(meta = (Experimental, DataflowGroom))
 struct FSmoothGuidesCurvesDataflowNode : public FDataflowNode
 {
 	GENERATED_BODY()
