Card generator dataflow nodes support for filter groups + individual card render

diff --git a/Engine/Plugins/Experimental/HairCardGenerator/Source/HairCardGeneratorDataflow/Private/BuildCardsSettingsNode.cpp b/Engine/Plugins/Experimental/HairCardGenerator/Source/HairCardGeneratorDataflow/Private/BuildCardsSettingsNode.cpp
--- a/Engine/Plugins/Experimental/HairCardGenerator/Source/HairCardGeneratorDataflow/Private/BuildCardsSettingsNode.cpp
+++ b/Engine/Plugins/Experimental/HairCardGenerator/Source/HairCardGeneratorDataflow/Private/BuildCardsSettingsNode.cpp
@@ -55,17 +55,36 @@ void FBuildCardsSettingsNode::Evaluate(UE::Dataflow::FContext& Context, const FD
 			FManagedArrayCollection GroomCollection = GetValue<FManagedArrayCollection>(Context, &Collection);
 			for(FHairGroupsCardsSourceDescription& HairCardsDescription : LocalAsset->GetHairGroupsCards())
 			{
-				HairCardsDescription.Textures = FHairGroupCardsTextures();
+				if(HairCardsDescription.LODIndex != INDEX_NONE)
+				{ 
+					HairCardsDescription.Textures = FHairGroupCardsTextures();
 				
-				OutputSettings.AddDefaulted();
-				FHairCardGeneratorUtils::BuildGenerationSettings(false, LocalAsset, HairCardsDescription,
-					OutputSettings.Last().GenerationSettings, OutputSettings.Last().GenerationFlags, OutputSettings.Last().PipelineFlags);
+					OutputSettings.AddDefaulted();
+					FHairCardGeneratorUtils::BuildGenerationSettings(false, LocalAsset, HairCardsDescription,
+						OutputSettings.Last().GenerationSettings, OutputSettings.Last().GenerationFlags, OutputSettings.Last().PipelineFlags);
+				
+					if(!FilterSettings.IsEmpty())
+					{
+						TArray<TArray<FName>> FilterCardGroups;
+						TArray<FName> FilterGroupNames;
+						for(const FGroomFilterSettings& LocalSettings : FilterSettings)
+						{
+							if((LocalSettings.GroupIndex == HairCardsDescription.GroupIndex) &&
+								(LocalSettings.LODIndex == HairCardsDescription.LODIndex) && (HairCardsDescription.LODIndex != INDEX_NONE))
+							{
+								FilterCardGroups.Add(LocalSettings.CardGroups);
+								FilterGroupNames.Add(LocalSettings.FilterName);
+							}
+						}
+						OutputSettings.Last().GenerationSettings->BuildFilterGroupSettings(FilterCardGroups, FilterGroupNames);
+					}
 
-				FHairCardGeneratorUtils::LoadGroomStrands(LocalAsset, [&GroomCollection](TArray<TArray<FVector>>& StrandsPositions)
-				{
-					UE::CardGen::Private::BuildStrandsPositions(GroomCollection, StrandsPositions);
-				});
-				OutputSettings.Last().GenerationSettings->UpdateStrandFilterAssignment();
+					FHairCardGeneratorUtils::LoadGroomStrands(LocalAsset, [&GroomCollection](TArray<TArray<FVector>>& StrandsPositions)
+					{
+						UE::CardGen::Private::BuildStrandsPositions(GroomCollection, StrandsPositions);
+					});
+					OutputSettings.Last().GenerationSettings->UpdateStrandFilterAssignment();
+				}
 			}
 		}
 		SetValue(Context, MoveTemp(OutputSettings), &CardsSettings);
diff --git a/Engine/Plugins/Experimental/HairCardGenerator/Source/HairCardGeneratorDataflow/Private/CardsAssetTerminalNode.cpp b/Engine/Plugins/Experimental/HairCardGenerator/Source/HairCardGeneratorDataflow/Private/CardsAssetTerminalNode.cpp
--- a/Engine/Plugins/Experimental/HairCardGenerator/Source/HairCardGeneratorDataflow/Private/CardsAssetTerminalNode.cpp
+++ b/Engine/Plugins/Experimental/HairCardGenerator/Source/HairCardGeneratorDataflow/Private/CardsAssetTerminalNode.cpp
@@ -22,15 +22,25 @@ void FCardsAssetTerminalNode::SetAssetValue(TObjectPtr<UObject> Asset, UE::Dataf
 	{
 		TArray<FGroomCardsSettings> LocalSettings = GetValue<TArray<FGroomCardsSettings>>(Context, &CardsSettings);
 
-		if(GroomAsset && !GroomAsset->GetHairGroupsCards().IsEmpty() && (LocalSettings.Num() == GroomAsset->GetHairGroupsCards().Num()))
+		if(GroomAsset && !GroomAsset->GetHairGroupsCards().IsEmpty())
 		{
-			GroomAsset->Modify();
-			int32 GroupIndex = 0;
-			for(FHairGroupsCardsSourceDescription& HairCardsDescription : GroomAsset->GetHairGroupsCards())
+			int32 SettingsIndex = 0;
+			for(int32 DescriptionIndex = 0, NumDescriptions = GroomAsset->GetHairGroupsCards().Num(); DescriptionIndex < NumDescriptions; ++DescriptionIndex)
 			{
-				FHairCardGeneratorUtils::BuildCardsAssets(GroomAsset, HairCardsDescription,
-					LocalSettings[GroupIndex].GenerationSettings,  LocalSettings[GroupIndex].GenerationFlags);
-				++GroupIndex;
+				// Use a copy so we can only apply changes on success
+				FHairGroupsCardsSourceDescription HairCardsDescription = GroomAsset->GetHairGroupsCards()[DescriptionIndex];
+				if(HairCardsDescription.LODIndex != INDEX_NONE && LocalSettings.IsValidIndex(SettingsIndex))
+				{ 
+					const bool bSuccess = FHairCardGeneratorUtils::BuildCardsAssets(GroomAsset, HairCardsDescription,
+						LocalSettings[SettingsIndex].GenerationSettings, LocalSettings[SettingsIndex].GenerationFlags);
+
+					if (bSuccess)
+					{
+						GroomAsset->Modify();
+						GroomAsset->GetHairGroupsCards()[DescriptionIndex] = HairCardsDescription;
+					}
+					++SettingsIndex;
+				}
 			}
 		}
 	}
diff --git a/Engine/Plugins/Experimental/HairCardGenerator/Source/HairCardGeneratorDataflow/Private/GenerateCardsClumpsNode.cpp b/Engine/Plugins/Experimental/HairCardGenerator/Source/HairCardGeneratorDataflow/Private/GenerateCardsClumpsNode.cpp
--- a/Engine/Plugins/Experimental/HairCardGenerator/Source/HairCardGeneratorDataflow/Private/GenerateCardsClumpsNode.cpp
+++ b/Engine/Plugins/Experimental/HairCardGenerator/Source/HairCardGeneratorDataflow/Private/GenerateCardsClumpsNode.cpp
@@ -24,17 +24,14 @@ void FGenerateCardsClumpsNode::Evaluate(UE::Dataflow::FContext& Context, const F
 			{
 				for(const FCardsClumpsSettings& OverideSettings : ClumpsSettings)
 				{
-					if(GenerationSettings->GetLODIndex() == OverideSettings.LODIndex || OverideSettings.LODIndex == INDEX_NONE)
+					for(TObjectPtr<UHairCardGeneratorGroupSettings> FilterSettings : GenerationSettings->GetFilterGroupSettings())
 					{
-						for(TObjectPtr<UHairCardGeneratorGroupSettings> FilterSettings : GenerationSettings->GetFilterGroupSettings())
+						if(FilterSettings.Get())
 						{
-							if(FilterSettings.Get())
+							if((OverideSettings.FilterName != NAME_None && FilterSettings->GetFilterName() == OverideSettings.FilterName) || OverideSettings.FilterName == NAME_None)
 							{
-								if((OverideSettings.CardGroup != NAME_None && FilterSettings->ApplyToCardGroups.Contains(OverideSettings.CardGroup)) || OverideSettings.CardGroup == NAME_None)
-								{
-									FilterSettings->TargetNumberOfCards = OverideSettings.NumCards;
-									FilterSettings->MaxNumberOfFlyaways = OverideSettings.NumFlyaways;
-								}
+								FilterSettings->TargetNumberOfCards = OverideSettings.NumCards;
+								FilterSettings->MaxNumberOfFlyaways = OverideSettings.NumFlyaways;
 							}
 						}
 					}
@@ -56,15 +53,30 @@ void FGenerateCardsClumpsNode::Evaluate(UE::Dataflow::FContext& Context, const F
 					if(FHairCardGeneratorUtils::LoadGenerationSettings(GenerationSettings))
 					{
 						TArray<int32> StrandsClumps;
-						int32 NumClumps;
+						int32 NumClumps = 0;
+
+						StrandsClumps.Init(INDEX_NONE, GroomCollection.NumElements(UE::Groom::FGroomStrandsFacade::CurvesGroup));
 						
 						bool bHasClumps = FHairCardGeneratorUtils::RunCardsGeneration(GenerationSettings, LODSettings.PipelineFlags,
 				[&StrandsClumps, &NumClumps](const TObjectPtr<const UHairCardGeneratorPluginSettings>& GenerationSettings, const int32 FilterIndex, const uint8 GenFlags)
 						{
-							return FHairCardGeneratorUtils::GenerateCardsClumps(GenerationSettings, FilterIndex, GenFlags, StrandsClumps, NumClumps);
+							TArray<int32> FilterClumps; int32 ClumpCount = 0;
+							if(FHairCardGeneratorUtils::GenerateCardsClumps(GenerationSettings, FilterIndex, GenFlags, FilterClumps, ClumpCount))
+							{ 
+								for (int32 CurveIndex = 0, NumCurves = FilterClumps.Num()-1; CurveIndex < NumCurves; ++CurveIndex)
+								{
+									if (StrandsClumps.IsValidIndex(CurveIndex) && (FilterClumps[CurveIndex] != INDEX_NONE))
+									{
+										StrandsClumps[CurveIndex] = FilterClumps[CurveIndex] + NumClumps;
+									}
+								}
+								NumClumps += ClumpCount;
+								return true;
+							}
+							return false;
 						}, false);
 			
-						if(bHasClumps)
+ 						if(bHasClumps)
 						{
 							FString ClumpIndicesLOD = CurveClumpIndicesAttribute.ToString();
 							ClumpIndicesLOD.AppendInt(GenerationSettings->GetLODIndex());
@@ -75,12 +87,9 @@ void FGenerateCardsClumpsNode::Evaluate(UE::Dataflow::FContext& Context, const F
 							TManagedArray<int32>& CurveClumpIndices = GroomCollection.AddAttribute<int32>(FName(ClumpIndicesLOD), UE::Groom::FGroomStrandsFacade::CurvesGroup);
 							TManagedArray<int32>& ObjectNumClumps = GroomCollection.AddAttribute<int32>(FName(NumClumpsLOD), UE::Groom::FGroomStrandsFacade::ObjectsGroup);
 							
-							if(CurveClumpIndices.Num() == (StrandsClumps.Num()-1))
+							for(int32 CurveIndex = 0, NumCurves = CurveClumpIndices.Num(); CurveIndex < NumCurves; ++CurveIndex)
 							{
-								for(int32 CurveIndex = 0, NumCurves = CurveClumpIndices.Num(); CurveIndex < NumCurves; ++CurveIndex)
-								{
-									CurveClumpIndices[CurveIndex] = StrandsClumps[CurveIndex];
-								}
+								CurveClumpIndices[CurveIndex] = StrandsClumps[CurveIndex];
 							}
 							for(int32 ObjectIndex = 0, NumObjects = ObjectNumClumps.Num(); ObjectIndex < NumObjects; ++ObjectIndex)
 							{
diff --git a/Engine/Plugins/Experimental/HairCardGenerator/Source/HairCardGeneratorDataflow/Private/GenerateCardsGeometryNode.cpp b/Engine/Plugins/Experimental/HairCardGenerator/Source/HairCardGeneratorDataflow/Private/GenerateCardsGeometryNode.cpp
--- a/Engine/Plugins/Experimental/HairCardGenerator/Source/HairCardGeneratorDataflow/Private/GenerateCardsGeometryNode.cpp
+++ b/Engine/Plugins/Experimental/HairCardGenerator/Source/HairCardGeneratorDataflow/Private/GenerateCardsGeometryNode.cpp
@@ -27,16 +27,13 @@ void FGenerateCardsGeometryNode::Evaluate(UE::Dataflow::FContext& Context, const
 			{
 				for(const FCardsGeometrySettings& OverideSettings : GeometrySettings)
 				{
-					if(GenerationSettings->GetLODIndex() == OverideSettings.LODIndex || OverideSettings.LODIndex == INDEX_NONE)
+					for(TObjectPtr<UHairCardGeneratorGroupSettings> FilterSettings : GenerationSettings->GetFilterGroupSettings())
 					{
-						for(TObjectPtr<UHairCardGeneratorGroupSettings> FilterSettings : GenerationSettings->GetFilterGroupSettings())
-						{
-							if(FilterSettings.Get())
-                        	{
-								if((OverideSettings.CardGroup != NAME_None && FilterSettings->ApplyToCardGroups.Contains(OverideSettings.CardGroup)) || OverideSettings.CardGroup == NAME_None)
-								{
-									FilterSettings->TargetTriangleCount = OverideSettings.NumTriangles;
-								}
+						if(FilterSettings.Get())
+                        {
+							if((OverideSettings.FilterName != NAME_None && FilterSettings->GetFilterName() == OverideSettings.FilterName) || OverideSettings.FilterName == NAME_None)
+							{
+								FilterSettings->TargetTriangleCount = OverideSettings.NumTriangles;
 							}
 						}
 					}
@@ -62,7 +59,13 @@ void FGenerateCardsGeometryNode::Evaluate(UE::Dataflow::FContext& Context, const
 						bool bHasGeometry = FHairCardGeneratorUtils::RunCardsGeneration(GenerationSettings, LODSettings.PipelineFlags,
                         	[&ClumpsGeometry](const TObjectPtr<const UHairCardGeneratorPluginSettings>& GenerationSettings, const int32 FilterIndex, const uint8 GenFlags)
 						{
-							return FHairCardGeneratorUtils::GenerateCardsGeometry(GenerationSettings, FilterIndex, GenFlags, ClumpsGeometry);
+                        	TArray<TArray<FVector3f>> FilterGeometry;
+                        	if(FHairCardGeneratorUtils::GenerateCardsGeometry(GenerationSettings, FilterIndex, GenFlags, FilterGeometry))
+                        	{
+								ClumpsGeometry.Append(FilterGeometry);
+                        		return true;
+                        	}
+							return false;
 						}, false);
    
 						if(bHasGeometry)
diff --git a/Engine/Plugins/Experimental/HairCardGenerator/Source/HairCardGeneratorDataflow/Private/GenerateCardsTexturesNode.cpp b/Engine/Plugins/Experimental/HairCardGenerator/Source/HairCardGeneratorDataflow/Private/GenerateCardsTexturesNode.cpp
--- a/Engine/Plugins/Experimental/HairCardGenerator/Source/HairCardGeneratorDataflow/Private/GenerateCardsTexturesNode.cpp
+++ b/Engine/Plugins/Experimental/HairCardGenerator/Source/HairCardGeneratorDataflow/Private/GenerateCardsTexturesNode.cpp
@@ -25,16 +25,13 @@ void FGenerateCardsTexturesNode::Evaluate(UE::Dataflow::FContext& Context, const
 			{
 				for(const FCardsTextureSettings& OverideSettings : TextureSettings)
 				{
-					if(GenerationSettings->GetLODIndex() == OverideSettings.LODIndex || OverideSettings.LODIndex == INDEX_NONE)
+					for(TObjectPtr<UHairCardGeneratorGroupSettings> FilterSettings : GenerationSettings->GetFilterGroupSettings())
 					{
-						for(TObjectPtr<UHairCardGeneratorGroupSettings> FilterSettings : GenerationSettings->GetFilterGroupSettings())
+						if (FilterSettings.Get())
 						{
-							if(FilterSettings.Get())
+							if ((OverideSettings.FilterName != NAME_None && FilterSettings->GetFilterName() == OverideSettings.FilterName) || OverideSettings.FilterName == NAME_None)
 							{
-								if((OverideSettings.CardGroup != NAME_None && FilterSettings->ApplyToCardGroups.Contains(OverideSettings.CardGroup)) || OverideSettings.CardGroup == NAME_None)
-								{
-									FilterSettings->NumberOfTexturesInAtlas = OverideSettings.NumTextures;
-								}
+								FilterSettings->NumberOfTexturesInAtlas = OverideSettings.NumTextures;
 							}
 						}
 					}
@@ -56,55 +53,70 @@ void FGenerateCardsTexturesNode::Evaluate(UE::Dataflow::FContext& Context, const
 					if(FHairCardGeneratorUtils::LoadGenerationSettings(GenerationSettings))
 					{
 						TArray<int32> CardsTextures;
-						int32 NumTextures = 0;
-						TArray<float> VertexUVs;
+						TArray<FVector2f> VertexUVs;
 						
 						bool bHasTextures = FHairCardGeneratorUtils::RunCardsGeneration(GenerationSettings, LODSettings.PipelineFlags,
-							[&CardsTextures, &NumTextures, &VertexUVs](const TObjectPtr<const UHairCardGeneratorPluginSettings>& GenerationSettings, const int32 FilterIndex, const uint8 GenFlags)
+							[&CardsTextures, &VertexUVs](const TObjectPtr<const UHairCardGeneratorPluginSettings>& GenerationSettings, const int32 FilterIndex, const uint8 GenFlags)
 						{
-							return FHairCardGeneratorUtils::GenerateCardsTextures(GenerationSettings, FilterIndex, GenFlags, CardsTextures, NumTextures, VertexUVs);
+							TArray<int32> FilterTextures;
+							int32 TexturesCount = 0;
+							if(FHairCardGeneratorUtils::GenerateCardsTexturesClusters(GenerationSettings, FilterIndex, GenFlags, FilterTextures, TexturesCount))
+							{
+								const int32 CardCount = FilterTextures.Num() - 1 - TexturesCount;
+								const int32 CardOffset = CardsTextures.Num();
+
+								CardsTextures.Reserve(CardOffset + CardCount);
+								for (int32 CardIndex = 0; CardIndex < CardCount; ++CardIndex)
+								{
+									CardsTextures.Add((FilterTextures[CardIndex] < TexturesCount) ? FilterTextures[CardCount + FilterTextures[CardIndex]] + CardOffset : INDEX_NONE);
+								}
+								return true;
+							}
+							return false;
 						}, false);
 			
 						if(bHasTextures)
 						{
-							const int32 NumCards = CardsTextures.Num()-1-NumTextures;
-							
-							FString CardsObjectsLODGroup = CardsObjectsGroup.ToString();
-							CardsObjectsLODGroup.AppendInt(GenerationSettings->GetLODIndex());
+							TArray<float> FilterUVs;
+							if(FHairCardGeneratorUtils::GenerateTexturesLayoutAndAtlases(GenerationSettings, LODSettings.PipelineFlags, FilterUVs))
+							{
+								VertexUVs.Reserve(FilterUVs.Num());
+								int32 UVsIndex = 0;
+								while (UVsIndex < FilterUVs.Num())
+								{
+									if (FilterUVs[UVsIndex] != -1.0)
+									{
+										VertexUVs.Add(FVector2f(FilterUVs[UVsIndex], FilterUVs[UVsIndex + 1]));
+										UVsIndex += 2;
+									}
+									else
+									{
+										UVsIndex += 1;
+									}
+								}
+								FString CardsObjectsLODGroup = CardsObjectsGroup.ToString();
+								CardsObjectsLODGroup.AppendInt(GenerationSettings->GetLODIndex());
 			
-							FString CardsVerticesLODGroup = FGenerateCardsGeometryNode::CardsVerticesGroup.ToString();
-							CardsVerticesLODGroup.AppendInt(GenerationSettings->GetLODIndex());
+								FString CardsVerticesLODGroup = FGenerateCardsGeometryNode::CardsVerticesGroup.ToString();
+								CardsVerticesLODGroup.AppendInt(GenerationSettings->GetLODIndex());
 							
-							TManagedArray<int32>& ObjectTextureIndices = GroomCollection.AddAttribute<int32>(ObjectTextureIndicesAttribute, FName(CardsObjectsLODGroup));
-							TManagedArray<FVector2f>& VertexTextureUvs = GroomCollection.AddAttribute<FVector2f>(VertexTextureUVsAttribute, FName(CardsVerticesLODGroup));
+								TManagedArray<int32>& ObjectTextureIndices = GroomCollection.AddAttribute<int32>(ObjectTextureIndicesAttribute, FName(CardsObjectsLODGroup));
+								TManagedArray<FVector2f>& VertexTextureUvs = GroomCollection.AddAttribute<FVector2f>(VertexTextureUVsAttribute, FName(CardsVerticesLODGroup));
 							
-							GroomCollection.EmptyGroup(FName(CardsObjectsLODGroup));
-							GroomCollection.AddElements(NumCards, FName(CardsObjectsLODGroup));
+								GroomCollection.EmptyGroup(FName(CardsObjectsLODGroup));
+								GroomCollection.AddElements(CardsTextures.Num(), FName(CardsObjectsLODGroup));
 							
-							for(int32 CardIndex = 0; CardIndex < NumCards; ++CardIndex)
-							{
-								ObjectTextureIndices[CardIndex] = (CardsTextures[CardIndex] < NumTextures) ? CardsTextures[NumCards+CardsTextures[CardIndex]] : INDEX_NONE;
-							}
-			
-							TArray<FVector2f> UVsCoords;
-							int32 UVsIndex = 0;
-							while(UVsIndex < VertexUVs.Num())
-							{
-								if(VertexUVs[UVsIndex] != -1.0)
+								for(int32 CardIndex = 0, NumCards = CardsTextures.Num(); CardIndex < NumCards; ++CardIndex)
 								{
-									UVsCoords.Add(FVector2f(VertexUVs[UVsIndex], VertexUVs[UVsIndex+1]));
-									UVsIndex += 2;
+									ObjectTextureIndices[CardIndex] = CardsTextures[CardIndex];
 								}
-								else
-								{
-									UVsIndex += 1;
-								}
-							}
-							if(UVsCoords.Num() == VertexTextureUvs.Num())
-							{
-								for(int32 VertexIndex = 0; VertexIndex < UVsCoords.Num(); ++VertexIndex)
+
+								if(VertexUVs.Num() == VertexTextureUvs.Num())
 								{
-									VertexTextureUvs[VertexIndex] = UVsCoords[VertexIndex];
+									for(int32 VertexIndex = 0, NumVertices = VertexUVs.Num(); VertexIndex < NumVertices; ++VertexIndex)
+									{
+										VertexTextureUvs[VertexIndex] = VertexUVs[VertexIndex];
+									}
 								}
 							}
 						}
diff --git a/Engine/Plugins/Experimental/HairCardGenerator/Source/HairCardGeneratorDataflow/Private/HairCardDataflowRendering.cpp b/Engine/Plugins/Experimental/HairCardGenerator/Source/HairCardGeneratorDataflow/Private/HairCardDataflowRendering.cpp
--- a/Engine/Plugins/Experimental/HairCardGenerator/Source/HairCardGeneratorDataflow/Private/HairCardDataflowRendering.cpp
+++ b/Engine/Plugins/Experimental/HairCardGenerator/Source/HairCardGeneratorDataflow/Private/HairCardDataflowRendering.cpp
@@ -32,9 +32,7 @@ namespace UE::CardGen::Private
 			const FManagedArrayCollection Default;
 			checkf(State.GetRenderOutputs().Num() == 1, TEXT("Expected FGraphRenderingState object to have one render output"));
 			const FName PrimaryOutput = State.GetRenderOutputs()[0];
-			FManagedArrayCollection Collection = State.GetValue<FManagedArrayCollection>(PrimaryOutput, Default);
-			
-			const TSharedRef<const FManagedArrayCollection> GroomCollection = MakeShared<FManagedArrayCollection>(MoveTemp(Collection));
+			const FManagedArrayCollection& Collection = State.GetValue<FManagedArrayCollection>(PrimaryOutput, Default);
 
 			FString CardsVerticesLODGroup = FGenerateCardsGeometryNode::CardsVerticesGroup.ToString();
 			CardsVerticesLODGroup.AppendInt(GGroomDataflowCardsLod);
@@ -42,64 +40,102 @@ namespace UE::CardGen::Private
 			FString CardsFacesLODGroup = FGenerateCardsGeometryNode::CardsFacesGroup.ToString();
 			CardsFacesLODGroup.AppendInt(GGroomDataflowCardsLod);
 
-			if (GroomCollection->HasAttribute(FGenerateCardsGeometryNode::VertexClumpPositionsAttribute, FName(CardsVerticesLODGroup)) &&
-				GroomCollection->HasAttribute(FGenerateCardsGeometryNode::FaceVertexIndicesAttribute, FName(CardsFacesLODGroup)))
+			if (Collection.HasAttribute(FGenerateCardsGeometryNode::VertexClumpPositionsAttribute, FName(CardsVerticesLODGroup)) &&
+				Collection.HasAttribute(FGenerateCardsGeometryNode::FaceVertexIndicesAttribute, FName(CardsFacesLODGroup)) && 
+				Collection.HasAttribute(FGenerateCardsGeometryNode::VertexCardIndicesAttribute, FName(CardsVerticesLODGroup)))
 			{
-				TArray<FVector3f> VertexPositions;
-				TArray<FVector3f> VertexNormals;
-				TArray<FIntVector> FaceIndices;
-			
 				if (State.GetViewMode().GetName() == Dataflow::FDataflowConstruction3DViewMode::Name)
 				{
-					const TManagedArray<FVector3f>& VertexPositionsArray = GroomCollection->GetAttribute<FVector3f>(
+					const TManagedArray<FVector3f>& VertexGlobalPositions = Collection.GetAttribute<FVector3f>(
 						FGenerateCardsGeometryNode::VertexClumpPositionsAttribute, FName(CardsVerticesLODGroup));
 						
-					const TManagedArray<FIntVector3>& FaceIndicesArray = GroomCollection->GetAttribute<FIntVector3>(
+					const TManagedArray<FIntVector3>& FaceGlobalVertices = Collection.GetAttribute<FIntVector3>(
 						FGenerateCardsGeometryNode::FaceVertexIndicesAttribute, FName(CardsFacesLODGroup));
 
-					VertexPositions = VertexPositionsArray.GetConstArray();
-					FaceIndices = FaceIndicesArray.GetConstArray();
+					const TManagedArray<int32>& VertexCardIndices = Collection.GetAttribute<int32>(
+						FGenerateCardsGeometryNode::VertexCardIndicesAttribute, FName(CardsVerticesLODGroup));
 
-					TArray<int32> NumFaces;
-					
-					VertexNormals.Init(FVector3f::Zero(), VertexPositions.Num());
-					NumFaces.Init(0, VertexPositions.Num());
-					
-					for(int32 FaceIndex = 0; FaceIndex < FaceIndices.Num(); ++FaceIndex)
+					int32 NumCards = 0;
+					for(const int32& CardIndex : VertexCardIndices.GetConstArray())
 					{
-						const FVector3f FaceNormal =  (VertexPositions[FaceIndices[FaceIndex][2]] - VertexPositions[FaceIndices[FaceIndex][0]]).Cross(
-							VertexPositions[FaceIndices[FaceIndex][1]] - VertexPositions[FaceIndices[FaceIndex][0]]).GetSafeNormal();
-						
-						VertexNormals[FaceIndices[FaceIndex][0]] += FaceNormal;
-						VertexNormals[FaceIndices[FaceIndex][1]] += FaceNormal;
-						VertexNormals[FaceIndices[FaceIndex][2]] += FaceNormal;
+						NumCards = FMath::Max(NumCards, CardIndex);
+					}
+					++NumCards;
 
-						NumFaces[FaceIndices[FaceIndex][0]]++;
-						NumFaces[FaceIndices[FaceIndex][1]]++;
-						NumFaces[FaceIndices[FaceIndex][2]]++;
+					TArray<TArray<int32>> CardVertexIndices;
+					TArray<TArray<int32>> CardFaceIndices;
+					TArray<int32> VertexLocalIndices;
+
+					CardVertexIndices.SetNum(NumCards);
+					CardFaceIndices.SetNum(NumCards);
+					VertexLocalIndices.SetNum(VertexCardIndices.Num());
+					
+					for (int32 VertexIndex = 0, NumVertices = VertexCardIndices.Num(); VertexIndex < NumVertices; ++VertexIndex)
+					{
+						VertexLocalIndices[VertexIndex] = CardVertexIndices[VertexCardIndices[VertexIndex]].Num();
+						CardVertexIndices[VertexCardIndices[VertexIndex]].Add(VertexIndex);
 					}
 
-					for(int32 VertexIndex = 0; VertexIndex < VertexPositions.Num(); ++VertexIndex)
+					for (int32 FaceIndex = 0, NumFaces = FaceGlobalVertices.Num(); FaceIndex < NumFaces; ++FaceIndex)
 					{
-						if(NumFaces[VertexIndex] > 0)
+						CardFaceIndices[VertexCardIndices[FaceGlobalVertices[FaceIndex][0]]].Add(FaceIndex);
+					}
+
+					for(int32 CardIndex = 0; CardIndex < NumCards; ++CardIndex)
+					{	
+						TArray<FVector3f> VertexLocalPositions;
+						TArray<FVector3f> VertexLocalNormals;
+						TArray<FIntVector> FaceLocalVertices;
+						TArray<int32> VertexNumFaces;
+
+						VertexLocalPositions.Init(FVector3f::Zero(), CardVertexIndices[CardIndex].Num());
+						VertexLocalNormals.Init(FVector3f::Zero(), CardVertexIndices[CardIndex].Num());
+						VertexNumFaces.Init(0, CardVertexIndices[CardIndex].Num());
+						FaceLocalVertices.Init(FIntVector::ZeroValue, CardFaceIndices.Num());
+
+						for (int32 FaceIndex = 0, NumFaces = CardFaceIndices[CardIndex].Num(); FaceIndex < NumFaces; ++FaceIndex)
+						{
+							const FIntVector3 GlobalVertices = FaceGlobalVertices[CardFaceIndices[CardIndex][FaceIndex]];
+							FaceLocalVertices[FaceIndex] = FIntVector(VertexLocalIndices[GlobalVertices[0]], VertexLocalIndices[GlobalVertices[1]], VertexLocalIndices[GlobalVertices[2]]);
+
+							const FVector3f FaceNormal = (VertexGlobalPositions[GlobalVertices[2]] - VertexGlobalPositions[GlobalVertices[0]]).Cross(
+								VertexGlobalPositions[GlobalVertices[1]] - VertexGlobalPositions[GlobalVertices[0]]).GetSafeNormal();
+
+							VertexLocalNormals[FaceLocalVertices[FaceIndex][0]] += FaceNormal;
+							VertexLocalNormals[FaceLocalVertices[FaceIndex][1]] += FaceNormal;
+							VertexLocalNormals[FaceLocalVertices[FaceIndex][2]] += FaceNormal;
+
+							VertexNumFaces[FaceLocalVertices[FaceIndex][0]]++;
+							VertexNumFaces[FaceLocalVertices[FaceIndex][1]]++;
+							VertexNumFaces[FaceLocalVertices[FaceIndex][2]]++;
+						}
+
+						for (int32 VertexIndex = 0, NumVertices = CardVertexIndices[CardIndex].Num(); VertexIndex < NumVertices; ++VertexIndex)
 						{
-							VertexNormals[VertexIndex] /= NumFaces[VertexIndex];
+							if (VertexNumFaces[VertexIndex] > 0)
+							{
+								VertexLocalNormals[VertexIndex] /= VertexNumFaces[VertexIndex];
+							}
+							VertexLocalPositions[VertexIndex] = VertexGlobalPositions[CardVertexIndices[CardIndex][VertexIndex]];
 						}
+
+						const FLinearColor GroupColor = FLinearColor::IntToDistinctColor(CardIndex, 0.75f, 1.0f, 90.0f);
+
+						TArray<FLinearColor> VertexLocalColors;
+						VertexLocalColors.Init(GroupColor, VertexLocalPositions.Num());
+
+						ComputeVertexColors(Collection, CardIndex, VertexLocalColors);
+
+						const FString GeometryName = TEXT("Groom_Card_") + FString::FromInt(CardIndex);
+						const int32 GeometryIndex = RenderCollection.StartGeometryGroup(GeometryName);
+						RenderCollection.AddSurface(MoveTemp(VertexLocalPositions), MoveTemp(FaceLocalVertices), MoveTemp(VertexLocalNormals), MoveTemp(VertexLocalColors));
+						RenderCollection.EndGeometryGroup(GeometryIndex);
 					}
 				}
 				else
 				{
-					checkf(false, TEXT("Invalid View Mode for FClothCollection rendering"));
+					checkf(false, TEXT("Invalid View Mode for groom cards geometry rendering"));
 				}
-			
-				TArray<FLinearColor> Colors;
-				Colors.Init(FLinearColor::Gray, VertexPositions.Num());	// TODO: Choose a vertex color
-
-				ComputeVertexColors(GroomCollection, Colors);
-			
-				const int32 GeometryIndex = RenderCollection.StartGeometryGroup(State.GetGuid().ToString());
-				RenderCollection.AddSurface(MoveTemp(VertexPositions), MoveTemp(FaceIndices), MoveTemp(VertexNormals), MoveTemp(Colors));
-				RenderCollection.EndGeometryGroup(GeometryIndex);
 			}
 		}
 	}
@@ -119,9 +155,7 @@ namespace UE::CardGen::Private
 				const FManagedArrayCollection Default;
 				checkf(State.GetRenderOutputs().Num() == 1, TEXT("Expected FGraphRenderingState object to have one render output"));
 				const FName PrimaryOutput = State.GetRenderOutputs()[0];
-				FManagedArrayCollection Collection = State.GetValue<FManagedArrayCollection>(PrimaryOutput, Default);
-			
-				const TSharedRef<const FManagedArrayCollection> GroomCollection = MakeShared<FManagedArrayCollection>(MoveTemp(Collection));
+				const FManagedArrayCollection& Collection = State.GetValue<FManagedArrayCollection>(PrimaryOutput, Default);
 
 				FString CardsVerticesLODGroup = FGenerateCardsGeometryNode::CardsVerticesGroup.ToString();
 				CardsVerticesLODGroup.AppendInt(GGroomDataflowCardsLod);
@@ -129,13 +163,13 @@ namespace UE::CardGen::Private
 				FString CardsFacesLODGroup = FGenerateCardsGeometryNode::CardsFacesGroup.ToString();
 				CardsFacesLODGroup.AppendInt(GGroomDataflowCardsLod);
 
-				if (GroomCollection->HasAttribute(FGenerateCardsTexturesNode::VertexTextureUVsAttribute, FName(CardsVerticesLODGroup)) &&
-					GroomCollection->HasAttribute(FGenerateCardsGeometryNode::FaceVertexIndicesAttribute, FName(CardsFacesLODGroup)))
+				if (Collection.HasAttribute(FGenerateCardsTexturesNode::VertexTextureUVsAttribute, FName(CardsVerticesLODGroup)) &&
+					Collection.HasAttribute(FGenerateCardsGeometryNode::FaceVertexIndicesAttribute, FName(CardsFacesLODGroup)))
 				{
-					const TManagedArray<FVector2f>& VertexUVsArray = GroomCollection->GetAttribute<FVector2f>(
+					const TManagedArray<FVector2f>& VertexUVsArray = Collection.GetAttribute<FVector2f>(
 							FGenerateCardsTexturesNode::VertexTextureUVsAttribute, FName(CardsVerticesLODGroup));
 						
-					const TManagedArray<FIntVector3>& FaceIndicesArray = GroomCollection->GetAttribute<FIntVector3>(
+					const TManagedArray<FIntVector3>& FaceIndicesArray = Collection.GetAttribute<FIntVector3>(
 						FGenerateCardsGeometryNode::FaceVertexIndicesAttribute, FName(CardsFacesLODGroup));
 					
 					TArray<FIntVector> FaceIndices = FaceIndicesArray.GetConstArray();
@@ -161,49 +195,23 @@ namespace UE::CardGen::Private
 			}
 		}
 	}
-	
-	void FCardsTextureRenderingCallbacks::ComputeVertexColors(const TSharedRef<const FManagedArrayCollection> GroomCollection, TArray<FLinearColor>& VertexColors) const
+
+	void FCardsTextureRenderingCallbacks::ComputeVertexColors(const FManagedArrayCollection& Collection, const int32 CardIndex, TArray<FLinearColor>& VertexColors) const
 	{
 		FString CardsObjectsLODGroup = FGenerateCardsTexturesNode::CardsObjectsGroup.ToString();
 		CardsObjectsLODGroup.AppendInt(GGroomDataflowCardsLod);
 
-		FString CardsVerticesLODGroup = FGenerateCardsGeometryNode::CardsVerticesGroup.ToString();
-		CardsVerticesLODGroup.AppendInt(GGroomDataflowCardsLod);
-
-		if (GroomCollection->HasAttribute(FGenerateCardsTexturesNode::ObjectTextureIndicesAttribute, FName(CardsObjectsLODGroup)) &&
-			GroomCollection->HasAttribute(FGenerateCardsGeometryNode::VertexCardIndicesAttribute,FName(CardsVerticesLODGroup)))
+		if (Collection.HasAttribute(FGenerateCardsTexturesNode::ObjectTextureIndicesAttribute, FName(CardsObjectsLODGroup)))
 		{
-			const TManagedArray<int32>& VertexCardsArray = GroomCollection->GetAttribute<int32>(
-					FGenerateCardsGeometryNode::VertexCardIndicesAttribute, FName(CardsVerticesLODGroup));
-						
-			const TManagedArray<int32>& CardsTextureArray = GroomCollection->GetAttribute<int32>(
-					FGenerateCardsTexturesNode::ObjectTextureIndicesAttribute, FName(CardsObjectsLODGroup));
-
-			int32 NumTextures = 0;
-			for(int32 CardTexture = 0; CardTexture < CardsTextureArray.Num(); ++CardTexture)
-			{
-				NumTextures = FMath::Max(NumTextures, CardsTextureArray[CardTexture]);
-			}
-			NumTextures += 1;
+			const TManagedArray<int32>& CardsTextureArray = Collection.GetAttribute<int32>(
+				FGenerateCardsTexturesNode::ObjectTextureIndicesAttribute, FName(CardsObjectsLODGroup));
 
-			TArray<FLinearColor> TextureColors;
-			TextureColors.Init(FLinearColor::Black, NumTextures);
+			const int32 TextureIndex = (CardIndex < CardsTextureArray.Num()) ? CardsTextureArray[CardIndex] : INDEX_NONE;
+			FLinearColor CardColor = (TextureIndex != INDEX_NONE) ? FLinearColor::IntToDistinctColor(TextureIndex, 0.75f, 1.0f, 90.0f) : FLinearColor::Black;
+			CardColor *= (TextureIndex == CardIndex) ? 1.0 : GGroomDataflowCardsAlpha;
 
-			for(int32 TextureIndex = 0; TextureIndex < NumTextures; ++TextureIndex)
-			{
-				TextureColors[TextureIndex] = FLinearColor::IntToDistinctColor(TextureIndex, 0.75f, 1.0f, 90.0f);
-			}
-
-			const int32 NumVertices = VertexCardsArray.Num();
-			VertexColors.Init(FLinearColor::Black, NumVertices);
-			
-			for(int32 VertexIndex = 0; VertexIndex < NumVertices; ++VertexIndex)
-			{
-				const int32 TextureIndex = (VertexCardsArray[VertexIndex] < CardsTextureArray.Num()) ?
-					CardsTextureArray[VertexCardsArray[VertexIndex]] : INDEX_NONE;
-				VertexColors[VertexIndex] = (TextureIndex != INDEX_NONE) ? TextureColors[TextureIndex] : FLinearColor::Black;
-				VertexColors[VertexIndex] *= (TextureIndex == VertexCardsArray[VertexIndex]) ? 1.0 : GGroomDataflowCardsAlpha;
-			}
+			const int32 NumVertices = VertexColors.Num();
+			VertexColors.Init(CardColor, NumVertices);
 		}
 	}
 
@@ -213,17 +221,20 @@ namespace UE::CardGen::Private
 	{
 		GroupAttribute = FGenerateCardsClumpsNode::CurveClumpIndicesAttribute.ToString();
 		GroupAttribute.AppendInt(GGroomDataflowCardsLod);
-
-		FString NumClumpsLOD = FGenerateCardsClumpsNode::ObjectNumClumpsAttribute.ToString();
-		NumClumpsLOD.AppendInt(GGroomDataflowCardsLod);
-
 		GroupName = FString("Clump");
 
 		const FManagedArrayCollection& GroomCollection = StrandsFacade.GetManagedArrayCollection();
-		if(GroomCollection.HasAttribute(FName(NumClumpsLOD), UE::Groom::FGroomStrandsFacade::ObjectsGroup))
+		if (GroomCollection.HasAttribute(FName(GroupAttribute), UE::Groom::FGroomStrandsFacade::CurvesGroup))
 		{
-			const TManagedArray<int32>& NumClumps = GroomCollection.GetAttribute<int32>(FName(NumClumpsLOD), UE::Groom::FGroomStrandsFacade::ObjectsGroup);
-			return (NumClumps.Num() > 0) ? NumClumps[0] : 0;
+			const TManagedArray<int32>& ClumpIndices = GroomCollection.GetAttribute<int32>(FName(GroupAttribute), UE::Groom::FGroomStrandsFacade::CurvesGroup);
+
+			int32 NumClumps = INDEX_NONE;
+			for(int32 ElementIndex = 0, NumElements = GroomCollection.NumElements(UE::Groom::FGroomStrandsFacade::CurvesGroup); ElementIndex < NumElements; ++ElementIndex)
+			{	
+				NumClumps = FMath::Max(NumClumps, ClumpIndices[ElementIndex]);
+			}
+			++NumClumps;
+			return NumClumps;
 		}
 		return 0;
 	}
diff --git a/Engine/Plugins/Experimental/HairCardGenerator/Source/HairCardGeneratorDataflow/Public/BuildCardsSettingsNode.h b/Engine/Plugins/Experimental/HairCardGenerator/Source/HairCardGeneratorDataflow/Public/BuildCardsSettingsNode.h
--- a/Engine/Plugins/Experimental/HairCardGenerator/Source/HairCardGeneratorDataflow/Public/BuildCardsSettingsNode.h
+++ b/Engine/Plugins/Experimental/HairCardGenerator/Source/HairCardGeneratorDataflow/Public/BuildCardsSettingsNode.h
@@ -17,22 +17,44 @@ struct FGroomCardsSettings
 	GENERATED_BODY()
 	
 	/** Generator settings to be built */
-	UPROPERTY(EditAnywhere, Category = "GroomCards");
+	UPROPERTY(EditAnywhere, Category = "Groom Cards");
 	TObjectPtr<UHairCardGeneratorPluginSettings> GenerationSettings = nullptr;
 
 	/** Generation flags to output the assets */
-	UPROPERTY(EditAnywhere, Category = "GroomCards");
+	UPROPERTY(EditAnywhere, Category = "Groom Cards");
 	uint8 GenerationFlags = 0;
 
 	/** Pipeline flags to generate clumps, geometry and textures */
-	UPROPERTY(EditAnywhere, Category = "GroomCards");
+	UPROPERTY(EditAnywhere, Category = "Groom Cards");
 	uint8 PipelineFlags = 0;
 
 	/** Groom asset to generate the */
-    UPROPERTY(EditAnywhere, Category = "GroomCards");
+    UPROPERTY(EditAnywhere, Category = "Groom Cards");
     TObjectPtr<UGroomAsset> GroomAsset = nullptr;
 };
 
+USTRUCT()
+struct FGroomFilterSettings
+{
+	GENERATED_BODY()
+
+	// Filter group name to be identified in the override settings
+	UPROPERTY(EditAnywhere, Category="Groom Cards")
+	FName FilterName;
+	
+	// LOD index of the filter group settings
+	UPROPERTY(EditAnywhere, Category="Groom Cards")
+	int32 LODIndex = -1;
+
+	// Group index of the filter group settings
+	UPROPERTY(EditAnywhere, Category="Groom Cards")
+	int32 GroupIndex = 0;
+
+	// Card group names that will belong to the filter group settings
+	UPROPERTY(EditAnywhere, Category="Groom Cards")
+	TArray<FName> CardGroups;
+};
+
 /** Build the cards generation settings */
 USTRUCT(meta = (Experimental, DataflowGroom))
 struct FBuildCardsSettingsNode : public FDataflowNode
@@ -56,15 +78,19 @@ public:
 	//~ End FDataflowNode interface
 
 	/** Groom asset to build the cards settings from */
-	UPROPERTY(EditAnywhere, Category = "GroomCards")
+	UPROPERTY(EditAnywhere, Category = "Groom Cards")
 	TObjectPtr<UGroomAsset> GroomAsset;
 	
 	/** Managed array collection to be used to store datas */
 	UPROPERTY(meta = (DataflowInput, DataflowOutput, DisplayName = "Collection", DataflowPassthrough  = "Collection"))
 	FManagedArrayCollection Collection;
 
-	/** Generator clumps settings to be built */
+	/** Generator cards settings to be built */
 	UPROPERTY(meta = (DataflowOutput));
 	TArray<FGroomCardsSettings> CardsSettings;
+	
+	/** List of filter setings to override */
+	UPROPERTY(EditAnywhere, Category = "Groom Cards")
+	TArray<FGroomFilterSettings> FilterSettings;
 };
 
diff --git a/Engine/Plugins/Experimental/HairCardGenerator/Source/HairCardGeneratorDataflow/Public/GenerateCardsClumpsNode.h b/Engine/Plugins/Experimental/HairCardGenerator/Source/HairCardGeneratorDataflow/Public/GenerateCardsClumpsNode.h
--- a/Engine/Plugins/Experimental/HairCardGenerator/Source/HairCardGeneratorDataflow/Public/GenerateCardsClumpsNode.h
+++ b/Engine/Plugins/Experimental/HairCardGenerator/Source/HairCardGeneratorDataflow/Public/GenerateCardsClumpsNode.h
@@ -15,21 +15,17 @@ USTRUCT()
 struct FCardsClumpsSettings
 {
 	GENERATED_BODY()
-	
-	// LOD index for which we will override the number of cards/triangles/texture/flyaways
-	UPROPERTY(EditAnywhere, Category="GroomCards")
-	int32 LODIndex = -1;
 
-	// Card group name for which we will override the number of cards/triangles/texture/flyaways
-	UPROPERTY(EditAnywhere, Category="GroomCards")
-	FName CardGroup = NAME_None;
+	// Card filter name for which we will override the number of cards/flyaways
+	UPROPERTY(EditAnywhere, Category="Groom Cards")
+	FName FilterName = NAME_None;
 
 	// Total number of cards to generate for this LOD settings
-	UPROPERTY(EditAnywhere, Category="GroomCards", meta=(ClampMin="1", ClampMax="100000"))
+	UPROPERTY(EditAnywhere, Category="Groom Cards", meta=(ClampMin="1", ClampMax="100000"))
 	int32 NumCards = 200;
 
 	// Maximum number of cards to assign to flyaway strands
-	UPROPERTY(EditAnywhere, Category="GroomCards", meta=(ClampMin="0", ClampMax="1000"))
+	UPROPERTY(EditAnywhere, Category="Groom Cards", meta=(ClampMin="0", ClampMax="1000"))
 	int32 NumFlyaways = 10;
 };
 
diff --git a/Engine/Plugins/Experimental/HairCardGenerator/Source/HairCardGeneratorDataflow/Public/GenerateCardsGeometryNode.h b/Engine/Plugins/Experimental/HairCardGenerator/Source/HairCardGeneratorDataflow/Public/GenerateCardsGeometryNode.h
--- a/Engine/Plugins/Experimental/HairCardGenerator/Source/HairCardGeneratorDataflow/Public/GenerateCardsGeometryNode.h
+++ b/Engine/Plugins/Experimental/HairCardGenerator/Source/HairCardGeneratorDataflow/Public/GenerateCardsGeometryNode.h
@@ -15,17 +15,13 @@ USTRUCT()
 struct FCardsGeometrySettings
 {
 	GENERATED_BODY()
-	
-	// LOD index for which we will override the number of cards/triangles/texture/flyaways
-	UPROPERTY(EditAnywhere, Category="GroomCards")
-	int32 LODIndex = -1;
 
-	// Card group name for which we will override the number of cards/triangles/texture/flyaways
-	UPROPERTY(EditAnywhere, Category="GroomCards")
-	FName CardGroup = NAME_None;
+	// Card filter name for which we will override the number of triangles
+	UPROPERTY(EditAnywhere, Category="Groom Cards")
+	FName FilterName = NAME_None;
 
 	// Total number of triangles to generate for this LOD settings
-	UPROPERTY(EditAnywhere, Category="GroomCards", meta=(ClampMin="1", ClampMax="100000"))
+	UPROPERTY(EditAnywhere, Category="Groom Cards", meta=(ClampMin="1", ClampMax="100000"))
 	int32 NumTriangles = 2000;
 };
 
diff --git a/Engine/Plugins/Experimental/HairCardGenerator/Source/HairCardGeneratorDataflow/Public/GenerateCardsTexturesNode.h b/Engine/Plugins/Experimental/HairCardGenerator/Source/HairCardGeneratorDataflow/Public/GenerateCardsTexturesNode.h
--- a/Engine/Plugins/Experimental/HairCardGenerator/Source/HairCardGeneratorDataflow/Public/GenerateCardsTexturesNode.h
+++ b/Engine/Plugins/Experimental/HairCardGenerator/Source/HairCardGeneratorDataflow/Public/GenerateCardsTexturesNode.h
@@ -15,17 +15,13 @@ USTRUCT()
 struct FCardsTextureSettings
 {
 	GENERATED_BODY()
-	
-	// LOD index for which we will override the number of cards/triangles/texture/flyaways
-	UPROPERTY(EditAnywhere, Category="GroomCards")
-	int32 LODIndex = -1;
 
-	// Card group name for which we will override the number of cards/triangles/texture/flyaways
-	UPROPERTY(EditAnywhere, Category="GroomCards")
-	FName CardGroup = NAME_None;
+	// Card filter name for which we will override the number of textures
+	UPROPERTY(EditAnywhere, Category="Groom Cards")
+	FName FilterName = NAME_None;
 
 	// Total number of triangles to generate for this LOD settings
-	UPROPERTY(EditAnywhere, Category="GroomCards")
+	UPROPERTY(EditAnywhere, Category="Groom Cards")
 	int32 NumTextures = 75;
 };
 
diff --git a/Engine/Plugins/Experimental/HairCardGenerator/Source/HairCardGeneratorDataflow/Public/HairCardDataflowRendering.h b/Engine/Plugins/Experimental/HairCardGenerator/Source/HairCardGeneratorDataflow/Public/HairCardDataflowRendering.h
--- a/Engine/Plugins/Experimental/HairCardGenerator/Source/HairCardGeneratorDataflow/Public/HairCardDataflowRendering.h
+++ b/Engine/Plugins/Experimental/HairCardGenerator/Source/HairCardGeneratorDataflow/Public/HairCardDataflowRendering.h
@@ -26,7 +26,7 @@ namespace UE::CardGen::Private
 		}
 
 		virtual void Render(GeometryCollection::Facades::FRenderingFacade& RenderCollection, const Dataflow::FGraphRenderingState& State) override;
-		virtual void ComputeVertexColors(const TSharedRef<const FManagedArrayCollection> GroomCollection, TArray<FLinearColor>& VertexColors) const {}
+		virtual void ComputeVertexColors(const FManagedArrayCollection& Collection, const int32 CardIndex, TArray<FLinearColor>& VertexColors) const {}
 	};
 
 	class FCardsClumpsRenderingCallbacks : public UE::Groom::FGroomStrandsRenderingCallbacks
@@ -64,7 +64,7 @@ namespace UE::CardGen::Private
 			return RenderKey;
 		}
 		virtual void Render(GeometryCollection::Facades::FRenderingFacade& RenderCollection, const Dataflow::FGraphRenderingState& State) override;
-		virtual void ComputeVertexColors(const TSharedRef<const FManagedArrayCollection> GroomCollection, TArray<FLinearColor>& VertexColors) const override;
+		virtual void ComputeVertexColors(const FManagedArrayCollection& Collection, const int32 CardIndex, TArray<FLinearColor>& VertexColors) const override;
 	};
 	
 	static void RegisterRenderingCallbacks()
diff --git a/Engine/Plugins/Experimental/HairCardGenerator/Source/HairCardGeneratorEditor/Private/HairCardGeneratorEditorModule.cpp b/Engine/Plugins/Experimental/HairCardGenerator/Source/HairCardGeneratorEditor/Private/HairCardGeneratorEditorModule.cpp
--- a/Engine/Plugins/Experimental/HairCardGenerator/Source/HairCardGeneratorEditor/Private/HairCardGeneratorEditorModule.cpp
+++ b/Engine/Plugins/Experimental/HairCardGenerator/Source/HairCardGeneratorEditor/Private/HairCardGeneratorEditorModule.cpp
@@ -246,8 +246,8 @@ bool FHairCardGeneratorUtils::GenerateCardsGeometry(const TObjectPtr<const UHair
 	return true;
 }
 
-bool FHairCardGeneratorUtils::GenerateCardsTextures(const TObjectPtr<const UHairCardGeneratorPluginSettings>& GenerationSettings, const int32 FilterIndex, const uint8 GenFlags,
-	TArray<int32>& CardsTextures, int32& NumTextures, TArray<float>& VertexUVs)
+bool FHairCardGeneratorUtils::GenerateCardsTexturesClusters(const TObjectPtr<const UHairCardGeneratorPluginSettings>& GenerationSettings, const int32 FilterIndex, const uint8 GenFlags,
+	TArray<int32>& CardsTextures, int32& NumTextures)
 {
 	TObjectPtr<UHairCardGenControllerBase> CardGenController = FHairCardGenController::GetHairCardGenerator();
 	if ( !CardGenController )
@@ -267,6 +267,16 @@ bool FHairCardGeneratorUtils::GenerateCardsTextures(const TObjectPtr<const UHair
 			NumTextures = CardsTextures.Last();
 		}
 	}
+	return true;
+}
+
+bool FHairCardGeneratorUtils::GenerateTexturesLayoutAndAtlases(const TObjectPtr<const UHairCardGeneratorPluginSettings>& GenerationSettings, const uint8 GenFlags, TArray<float>& VertexUVs)
+{
+	TObjectPtr<UHairCardGenControllerBase> CardGenController = FHairCardGenController::GetHairCardGenerator();
+	if ( !CardGenController )
+	{
+		return false;
+	}
 	if ( UHairCardGeneratorPluginSettings::CheckGenerationFlags(GenFlags, EHairCardGenerationPipeline::TextureLayout) && !GenerationSettings->bReduceCardsFromPreviousLOD )
 	{
 		VertexUVs = CardGenController->GenerateTextureLayout();
@@ -767,8 +777,7 @@ static bool HairCardGeneratorEditor_Impl::GenerateCardsForCardGroup(const TObjec
 	SlowTask.EnterProgressFrame(4);
 	TArray<int32> CardsTextures;
 	int32 NumTextures= 0;
-	TArray<float> VertexUVs;
-	if(!FHairCardGeneratorUtils::GenerateCardsTextures(GenerationSettings, FilterIndex, GenFlags, CardsTextures, NumTextures, VertexUVs))
+	if(!FHairCardGeneratorUtils::GenerateCardsTexturesClusters(GenerationSettings, FilterIndex, GenFlags, CardsTextures, NumTextures))
 	{
 		return false;
 	}
@@ -967,6 +976,12 @@ bool FHairCardGeneratorEditorModule::GenerateHairCardsForLOD(UGroomAsset* NewGro
 		return false;
 	}
 
+	TArray<float> VertexUVs;
+	if(!FHairCardGeneratorUtils::GenerateTexturesLayoutAndAtlases(GenerationSettings, AllGenFlags, VertexUVs))
+	{
+		return false;
+	}
+
 	if(!FHairCardGeneratorUtils::BuildCardsAssets(NewGroomAsset, CardsDesc, GenerationSettings, AllGenFlags))
 	{
 		return false;
diff --git a/Engine/Plugins/Experimental/HairCardGenerator/Source/HairCardGeneratorEditor/Private/HairCardGeneratorPluginSettings.cpp b/Engine/Plugins/Experimental/HairCardGenerator/Source/HairCardGeneratorEditor/Private/HairCardGeneratorPluginSettings.cpp
--- a/Engine/Plugins/Experimental/HairCardGenerator/Source/HairCardGeneratorEditor/Private/HairCardGeneratorPluginSettings.cpp
+++ b/Engine/Plugins/Experimental/HairCardGenerator/Source/HairCardGeneratorEditor/Private/HairCardGeneratorPluginSettings.cpp
@@ -917,6 +917,64 @@ void UHairCardGeneratorPluginSettings::ResetNumFilterGroups(int Count)
 	FilterGroupGenerationSettings.SetNum(Count, EAllowShrinking::Yes);
 }
 
+void UHairCardGeneratorPluginSettings::MakeSingleFilterGroupSettings()
+{
+	FilterGroupGenerationSettings.Reset();
+	
+	for(FName& CardGroup : CardGroupIds)
+	{
+		FilterGroupGenerationSettings.Add(NewObject<UHairCardGeneratorGroupSettings>(this));
+
+		FilterGroupGenerationSettings.Last()->SetSource(GroomAsset, LODIndex);
+		FilterGroupGenerationSettings.Last()->Reset(UHairCardGeneratorGroupSettings::StaticClass()->GetDefaultObject<UHairCardGeneratorGroupSettings>());
+
+		FilterGroupGenerationSettings.Last()->ResetFilterDefault({CardGroup});
+		FilterGroupGenerationSettings.Last()->UpdateStrandFilter();
+	}
+
+	ResetNumFilterGroups(CardGroupIds.Num());
+	PostResetUpdates();
+}
+
+void UHairCardGeneratorPluginSettings::BuildFilterGroupSettings(const TArray<TArray<FName>>& FilterCardGroups, const TArray<FName>& FilterGroupNames)
+{
+	if(FilterGroupNames.Num() == FilterCardGroups.Num())
+	{
+		FilterGroupGenerationSettings.Reset();
+		int32 FilterIndex = 0;
+		for(const TArray<FName>& CardGroups : FilterCardGroups)
+		{
+			TArray<FName> LocalGroups;
+			for(const FName& CardGroup : CardGroups)
+			{
+				if(CardGroupIds.Contains(CardGroup))
+				{
+					LocalGroups.Add(CardGroup);
+				}
+			}
+			FilterGroupGenerationSettings.Add(NewObject<UHairCardGeneratorGroupSettings>(this));
+
+			FilterGroupGenerationSettings.Last()->SetSource(GroomAsset, LODIndex);
+			FilterGroupGenerationSettings.Last()->Reset(UHairCardGeneratorGroupSettings::StaticClass()->GetDefaultObject<UHairCardGeneratorGroupSettings>());
+
+			if (!LocalGroups.IsEmpty())
+			{
+				FilterGroupGenerationSettings.Last()->ResetFilterDefault(LocalGroups);
+			}
+			else
+			{
+				FilterGroupGenerationSettings.Last()->ResetFilterDefault(CardGroupIds);
+			}
+			FilterGroupGenerationSettings.Last()->UpdateStrandFilter();
+			FilterGroupGenerationSettings.Last()->SetFilterName(FilterGroupNames[FilterIndex]);
+			++FilterIndex;
+		}
+
+		ResetNumFilterGroups(FilterGroupGenerationSettings.Num());
+		PostResetUpdates();
+	}
+}
+
 void UHairCardGeneratorPluginSettings::ResetFilterGroupSettings(int FilterGroupIndex, TObjectPtr<UHairCardGeneratorGroupSettings> OptionsTemplate)
 {
 	if ( FilterGroupGenerationSettings.Num() <= FilterGroupIndex )
diff --git a/Engine/Plugins/Experimental/HairCardGenerator/Source/HairCardGeneratorEditor/Private/HairCardGeneratorPluginSettings.h b/Engine/Plugins/Experimental/HairCardGenerator/Source/HairCardGeneratorEditor/Private/HairCardGeneratorPluginSettings.h
--- a/Engine/Plugins/Experimental/HairCardGenerator/Source/HairCardGeneratorEditor/Private/HairCardGeneratorPluginSettings.h
+++ b/Engine/Plugins/Experimental/HairCardGenerator/Source/HairCardGeneratorEditor/Private/HairCardGeneratorPluginSettings.h
@@ -191,6 +191,12 @@ public:
 	void AddNewSettingsFilterGroup();
 	void RemoveSettingsFilterGroup(TObjectPtr<UHairCardGeneratorGroupSettings> SettingsGroup);
 
+	/** Make filter settings for each groups */
+	void MakeSingleFilterGroupSettings();
+
+	/** Build the filter group settings according to a list of valid card groups and names */
+	void BuildFilterGroupSettings(const TArray<TArray<FName>>& FilterCardGroups, const TArray<FName>& FilterGroupNames);
+
 	// Output settings to json file
 	void WriteGenerationSettings();
 
@@ -362,6 +368,10 @@ private:
 
 	bool StrandFilterChecksDirty = false;
 	TUniquePtr<FComposableStrandFilterOp> StrandsFilter;
+
+	UPROPERTY(EditAnywhere, Category="Cards")
+	FName FilterName;
+	
 public:
 	// Specify which strands (card group ids) this settings group applies to
 	UPROPERTY(EditAnywhere, BlueprintReadOnly, Category="Cards")
@@ -426,6 +436,12 @@ public:
 
 	int32 GetStrandCount() {return StrandCount;}
 
+	/** Get the filter name */
+	const FName& GetFilterName() const {return FilterName;}
+
+	/** Set the filter name with a given identifier */
+	void SetFilterName(const FName& InFilterName) {FilterName = InFilterName;}
+
 private:
 	void SetSource(TObjectPtr<const UGroomAsset> SourceGroom, int32 GenLODIndex = 0);
 
diff --git a/Engine/Plugins/Experimental/HairCardGenerator/Source/HairCardGeneratorEditor/Public/HairCardGeneratorEditorModule.h b/Engine/Plugins/Experimental/HairCardGenerator/Source/HairCardGeneratorEditor/Public/HairCardGeneratorEditorModule.h
--- a/Engine/Plugins/Experimental/HairCardGenerator/Source/HairCardGeneratorEditor/Public/HairCardGeneratorEditorModule.h
+++ b/Engine/Plugins/Experimental/HairCardGenerator/Source/HairCardGeneratorEditor/Public/HairCardGeneratorEditorModule.h
@@ -28,8 +28,11 @@ struct HAIRCARDGENERATOREDITOR_API FHairCardGeneratorUtils
 	/** Generate the cards geometry based on the settings given a filter index */
 	static bool GenerateCardsGeometry(const TObjectPtr<const UHairCardGeneratorPluginSettings>& GenerationSettings, const int32 FilterIndex, const uint8 GenFlags, TArray<TArray<FVector3f>>& ClumpsGeometry);
 
-	/** Generate the cards texture based on the settings given a filter index */
-	static bool GenerateCardsTextures(const TObjectPtr<const UHairCardGeneratorPluginSettings>& GenerationSettings, const int32 FilterIndex, const uint8 GenFlags, TArray<int32>& CardsTextures, int32& NumTextures, TArray<float>& VertexUVs);
+	/** Generate the cards texture clusters based on the settings given a filter index */
+	static bool GenerateCardsTexturesClusters(const TObjectPtr<const UHairCardGeneratorPluginSettings>& GenerationSettings, const int32 FilterIndex, const uint8 GenFlags, TArray<int32>& CardsTextures, int32& NumTextures);
+
+	/** Generate the cards texture layout and atlases */
+	static bool GenerateTexturesLayoutAndAtlases(const TObjectPtr<const UHairCardGeneratorPluginSettings>& GenerationSettings, const uint8 GenFlags, TArray<float>& VertexUVs);
 	
 	/** Run the cards generation function */
 	static bool RunCardsGeneration(const TObjectPtr<UHairCardGeneratorPluginSettings>& GenerationSettings, const uint8 PipelineFlags, const FGeneratorFunction& PipelineFunction, const bool bCheckFlags = true);
diff --git a/Engine/Plugins/Runtime/HairStrands/Source/HairStrandsDataflow/Private/GroomDataflowRendering.cpp b/Engine/Plugins/Runtime/HairStrands/Source/HairStrandsDataflow/Private/GroomDataflowRendering.cpp
--- a/Engine/Plugins/Runtime/HairStrands/Source/HairStrandsDataflow/Private/GroomDataflowRendering.cpp
+++ b/Engine/Plugins/Runtime/HairStrands/Source/HairStrandsDataflow/Private/GroomDataflowRendering.cpp
@@ -1,4 +1,4 @@
-﻿// Copyright Epic Games, Inc. All Rights Reserved.
+// Copyright Epic Games, Inc. All Rights Reserved.
 
 #include "GroomDataflowRendering.h"
 #include "GroomCollectionFacades.h"
@@ -108,8 +108,8 @@ namespace UE::Groom
 	}
 
 	template<typename FacadeType>
-	FORCEINLINE void BuildRenderingGroups(const FacadeType& GroomFacade, TArray<TArray<FVector3f>>& GroupVertexPositions,
-		TArray<TArray<FVector3f>>& GroupVertexNormals, TArray<TArray<FIntVector>>& GroupFaceVertices, const FString& GroupAttribute, const int32)
+	static void BuildRenderingGroups(const FacadeType& GroomFacade, TArray<TArray<FVector3f>>& GroupVertexPositions,
+		TArray<TArray<FVector3f>>& GroupVertexNormals, TArray<TArray<FIntVector>>& GroupFaceVertices, const FString& GroupAttribute)
 	{
 		const int32 NumCurves = GroomFacade.GetNumCurves();
 
@@ -118,74 +118,83 @@ namespace UE::Groom
 		{
 			const TManagedArray<int32>& CurveGroups = GroomCollection.GetAttribute<int32>(FName(GroupAttribute), FacadeType::CurvesGroup);
 
-			int32 NumGroups = 0;
+			int32 NumGroups = INDEX_NONE;
 			for(int32 CurveIndex = 0; CurveIndex < NumCurves; ++CurveIndex)
 			{
 				NumGroups = FMath::Max(NumGroups, CurveGroups[CurveIndex]);
 			}
 			++NumGroups;
-			
-			GroupVertexPositions.SetNum(NumGroups);
-			GroupVertexNormals.SetNum(NumGroups);
-			GroupFaceVertices.SetNum(NumGroups);
 
-			TArray<int32> GroupsPoints;
-			GroupsPoints.Init(0, NumGroups);
+			if(NumGroups > 0)
+			{ 
+				GroupVertexPositions.SetNum(NumGroups);
+				GroupVertexNormals.SetNum(NumGroups);
+				GroupFaceVertices.SetNum(NumGroups);
+
+				TArray<int32> GroupsPoints;
+				GroupsPoints.Init(0, NumGroups);
 			
-			TArray<int32> GroupsCurves;
-			GroupsCurves.Init(0, NumGroups);
+				TArray<int32> GroupsCurves;
+				GroupsCurves.Init(0, NumGroups);
 
-			int32 PointOffset = 0;
-			for(int32 CurveIndex = 0; CurveIndex < NumCurves; ++CurveIndex)
-			{
-				const int32 CurvePoints =  GroomFacade.GetCurvePointOffsets()[CurveIndex] - PointOffset;
-				GroupsPoints[CurveGroups[CurveIndex]] += CurvePoints;
-				++GroupsCurves[CurveGroups[CurveIndex]];
-				PointOffset = GroomFacade.GetCurvePointOffsets()[CurveIndex];
-			}
+				int32 PointOffset = 0;
+				for(int32 CurveIndex = 0; CurveIndex < NumCurves; ++CurveIndex)
+				{
+					const int32 GroupIndex = CurveGroups[CurveIndex];
+					if(GroupsPoints.IsValidIndex(GroupIndex) && GroupsCurves.IsValidIndex(GroupIndex))
+					{ 
+						const int32 CurvePoints = GroomFacade.GetCurvePointOffsets()[CurveIndex] - PointOffset;
+						GroupsPoints[GroupIndex] += CurvePoints;
+						++GroupsCurves[GroupIndex];
+					}
+					PointOffset = GroomFacade.GetCurvePointOffsets()[CurveIndex];
+				}
 
-			for(int32 GroupIndex = 0; GroupIndex < NumGroups; ++GroupIndex)
-			{
-				GroupVertexPositions[GroupIndex].SetNum(GroupsPoints[GroupIndex] * 2);
-				GroupVertexNormals[GroupIndex].SetNum(GroupsPoints[GroupIndex] * 2);
-				GroupFaceVertices[GroupIndex].SetNum((GroupsPoints[GroupIndex]-GroupsCurves[GroupIndex]) * 2);
-			}
+				for(int32 GroupIndex = 0; GroupIndex < NumGroups; ++GroupIndex)
+				{
+					GroupVertexPositions[GroupIndex].SetNum(GroupsPoints[GroupIndex] * 2);
+					GroupVertexNormals[GroupIndex].SetNum(GroupsPoints[GroupIndex] * 2);
+					GroupFaceVertices[GroupIndex].SetNum((GroupsPoints[GroupIndex]-GroupsCurves[GroupIndex]) * 2);
+				}
 
-			GroupsPoints.Init(0, NumGroups);
-			GroupsCurves.Init(0, NumGroups);
+				GroupsPoints.Init(0, NumGroups);
+				GroupsCurves.Init(0, NumGroups);
 			
-			PointOffset = 0;
-			for(int32 CurveIndex = 0; CurveIndex < NumCurves; ++CurveIndex)
-			{
-				const int32 GroupIndex = CurveGroups[CurveIndex];
-				
-				TArray<FVector3f>& VertexPositions = GroupVertexPositions[GroupIndex];
-				TArray<FVector3f>& VertexNormals = GroupVertexNormals[GroupIndex];
-				TArray<FIntVector>& FaceVertices = GroupFaceVertices[GroupIndex];
-				
-				for (int32 PointIndex = PointOffset, PointEnd = GroomFacade.GetCurvePointOffsets()[CurveIndex];
-						   PointIndex < PointEnd; ++PointIndex)
+				PointOffset = 0;
+				for(int32 CurveIndex = 0; CurveIndex < NumCurves; ++CurveIndex)
 				{
-					// Build 2 vertices for each points
-					const int32 EdgeIndex = (PointIndex == (PointEnd-1)) ? PointIndex-CurveIndex-1 : PointIndex-CurveIndex;
+					const int32 GroupIndex = CurveGroups[CurveIndex];
+					if (GroupsPoints.IsValidIndex(GroupIndex) && GroupsCurves.IsValidIndex(GroupIndex))
+					{
+						TArray<FVector3f>& VertexPositions = GroupVertexPositions[GroupIndex];
+						TArray<FVector3f>& VertexNormals = GroupVertexNormals[GroupIndex];
+						TArray<FIntVector>& FaceVertices = GroupFaceVertices[GroupIndex];
+				
+						for (int32 PointIndex = PointOffset, PointEnd = GroomFacade.GetCurvePointOffsets()[CurveIndex];
+								   PointIndex < PointEnd; ++PointIndex)
+						{
+							// Build 2 vertices for each points
+							const int32 EdgeIndex = (PointIndex == (PointEnd-1)) ? PointIndex-CurveIndex-1 : PointIndex-CurveIndex;
 										
-					const FVector3f SideVector = GroomFacade.GetEdgeRestOrientations()[EdgeIndex].GetAxisX();
-					BuildVertexPositions(GroupsPoints[GroupIndex], GroomFacade.GetPointRestPositions()[PointIndex], SideVector, GGroomDataflowRenderingThickness, VertexPositions);
+							const FVector3f SideVector = GroomFacade.GetEdgeRestOrientations()[EdgeIndex].GetAxisX();
+							BuildVertexPositions(GroupsPoints[GroupIndex], GroomFacade.GetPointRestPositions()[PointIndex], SideVector, GGroomDataflowRenderingThickness, VertexPositions);
 
-					const FVector3f NormalVector = GroomFacade.GetEdgeRestOrientations()[EdgeIndex].GetAxisZ();
-					BuildVertexNormals(GroupsPoints[GroupIndex], NormalVector, VertexNormals);
+							const FVector3f NormalVector = GroomFacade.GetEdgeRestOrientations()[EdgeIndex].GetAxisZ();
+							BuildVertexNormals(GroupsPoints[GroupIndex], NormalVector, VertexNormals);
 							
-					// Build 2 faces for each edges perpendicular to the up-vector curve
-					if(PointIndex < (PointEnd-1))
-					{
-						BuildFaceVertices(GroupsPoints[GroupIndex], GroupsCurves[GroupIndex], FaceVertices);
-					}
+							// Build 2 faces for each edges perpendicular to the up-vector curve
+							if(PointIndex < (PointEnd-1))
+							{
+								BuildFaceVertices(GroupsPoints[GroupIndex], GroupsCurves[GroupIndex], FaceVertices);
+							}
 
-					++GroupsPoints[GroupIndex];
+							++GroupsPoints[GroupIndex];
 					
+						}
+						++GroupsCurves[GroupIndex];
+					}
+					PointOffset = GroomFacade.GetCurvePointOffsets()[CurveIndex];
 				}
-				++GroupsCurves[GroupIndex];
-				PointOffset = GroomFacade.GetCurvePointOffsets()[CurveIndex];
 			}
 		}
 	}
@@ -269,10 +278,9 @@ namespace UE::Groom
 			const FManagedArrayCollection Default;
 			checkf(State.GetRenderOutputs().Num() == 1, TEXT("Expected FGraphRenderingState object to have one render output"));
 			const FName PrimaryOutput = State.GetRenderOutputs()[0];
-			FManagedArrayCollection Collection = State.GetValue<FManagedArrayCollection>(PrimaryOutput, Default);
+			const FManagedArrayCollection& Collection = State.GetValue<FManagedArrayCollection>(PrimaryOutput, Default);
 
-			const TSharedRef<const FManagedArrayCollection> GroomCollection = MakeShared<FManagedArrayCollection>(MoveTemp(Collection));
-			const FacadeType GroomFacade(*GroomCollection);
+			const FacadeType GroomFacade(Collection);
 			if (GroomFacade.IsValid())
 			{
 				TArray<FVector3f> VertexPositions;
@@ -301,7 +309,7 @@ namespace UE::Groom
 	}
 
 	template<typename FacadeType>
-	FORCEINLINE void RenderGroupCollection(GeometryCollection::Facades::FRenderingFacade& RenderCollection, const Dataflow::FGraphRenderingState& State,
+	static void RenderGroupCollection(GeometryCollection::Facades::FRenderingFacade& RenderCollection, const Dataflow::FGraphRenderingState& State,
 		const TFunction<int32(const FacadeType&, FString&, FString&)>& GroupLambda)
 	{
 		if (State.GetRenderOutputs().Num())
@@ -309,10 +317,9 @@ namespace UE::Groom
 			const FManagedArrayCollection Default;
 			checkf(State.GetRenderOutputs().Num() == 1, TEXT("Expected FGraphRenderingState object to have one render output"));
 			const FName PrimaryOutput = State.GetRenderOutputs()[0];
-			FManagedArrayCollection Collection = State.GetValue<FManagedArrayCollection>(PrimaryOutput, Default);
+			const FManagedArrayCollection& Collection = State.GetValue<FManagedArrayCollection>(PrimaryOutput, Default);
 
-			const TSharedRef<const FManagedArrayCollection> GroomCollection = MakeShared<FManagedArrayCollection>(MoveTemp(Collection));
-			const FacadeType GroomFacade(*GroomCollection);
+			const FacadeType GroomFacade(Collection);
 			if (GroomFacade.IsValid())
 			{
 				FString GroupAttribute, GroupName;
@@ -324,25 +331,27 @@ namespace UE::Groom
 
 				if (State.GetViewMode().GetName() == Dataflow::FDataflowConstruction3DViewMode::Name)
 				{
-					BuildRenderingGroups(GroomFacade, VertexPositions, VertexNormals, FaceIndices, GroupAttribute, NumGroups);
+					BuildRenderingGroups(GroomFacade, VertexPositions, VertexNormals, FaceIndices, GroupAttribute);
 				}
 				else
 				{
 					UE_LOG(LogTemp, Error, TEXT("Invalid View Mode for Groom dataflow rendering"));
 				}
+				if(VertexPositions.Num() == NumGroups && VertexNormals.Num() == NumGroups && FaceIndices.Num() == NumGroups)
+				{ 
+					for (int32 GroupIndex = 0; GroupIndex < NumGroups; ++GroupIndex)
+					{
+						const FLinearColor GroupColor = FLinearColor::IntToDistinctColor(GroupIndex, 0.75f, 1.0f, 90.0f);
 
-				for(int32 GroupIndex = 0; GroupIndex < NumGroups; ++GroupIndex)
-				{
-					const FLinearColor GroupColor = FLinearColor::IntToDistinctColor(GroupIndex, 0.75f, 1.0f, 90.0f);
-					
-					TArray<FLinearColor> VertexColors;
-					VertexColors.Init(GroupColor, VertexPositions[GroupIndex].Num());
-					
-					const FString GeometryName = GroomFacade.GetObjectGroupNames().Last() + TEXT("_") + GroupName + TEXT("_") + FString::FromInt(GroupIndex);
-					const int32 GeometryIndex = RenderCollection.StartGeometryGroup(GeometryName);
-					RenderCollection.AddSurface(MoveTemp(VertexPositions[GroupIndex]),
-						MoveTemp(FaceIndices[GroupIndex]), MoveTemp(VertexNormals[GroupIndex]), MoveTemp(VertexColors));
-					RenderCollection.EndGeometryGroup(GeometryIndex);
+						TArray<FLinearColor> VertexColors;
+						VertexColors.Init(GroupColor, VertexPositions[GroupIndex].Num());
+
+						const FString GeometryName = GroomFacade.GetObjectGroupNames().Last() + TEXT("_") + GroupName + TEXT("_") + FString::FromInt(GroupIndex);
+						const int32 GeometryIndex = RenderCollection.StartGeometryGroup(GeometryName);
+						RenderCollection.AddSurface(MoveTemp(VertexPositions[GroupIndex]),
+							MoveTemp(FaceIndices[GroupIndex]), MoveTemp(VertexNormals[GroupIndex]), MoveTemp(VertexColors));
+						RenderCollection.EndGeometryGroup(GeometryIndex);
+					}
 				}
 			}
 		}
