Groom clumps dataflow rendering

diff --git a/Engine/Plugins/Experimental/HairCardGenerator/Source/HairCardGeneratorDataflow/HairCardGeneratorDataflow.Build.cs b/Engine/Plugins/Experimental/HairCardGenerator/Source/HairCardGeneratorDataflow/HairCardGeneratorDataflow.Build.cs
--- a/Engine/Plugins/Experimental/HairCardGenerator/Source/HairCardGeneratorDataflow/HairCardGeneratorDataflow.Build.cs
+++ b/Engine/Plugins/Experimental/HairCardGenerator/Source/HairCardGeneratorDataflow/HairCardGeneratorDataflow.Build.cs
@@ -23,9 +23,11 @@ namespace UnrealBuildTool.Rules
 					"DataflowNodes",
 					"HairStrandsCore",
 					"HairStrandsDataflow",
+					"HairStrandsEditor",
 					"HairCardGeneratorEditor",
 					"HairCardGeneratorFramework",
 					"Engine",
+					"GeometryCore"
 				}
 			);
 		}
diff --git a/Engine/Plugins/Experimental/HairCardGenerator/Source/HairCardGeneratorDataflow/Private/GenerateCardsClumpsNode.cpp b/Engine/Plugins/Experimental/HairCardGenerator/Source/HairCardGeneratorDataflow/Private/GenerateCardsClumpsNode.cpp
--- a/Engine/Plugins/Experimental/HairCardGenerator/Source/HairCardGeneratorDataflow/Private/GenerateCardsClumpsNode.cpp
+++ b/Engine/Plugins/Experimental/HairCardGenerator/Source/HairCardGeneratorDataflow/Private/GenerateCardsClumpsNode.cpp
@@ -4,6 +4,9 @@
 #include "GroomCollectionFacades.h"
 #include "HairCardGeneratorEditorModule.h"
 #include "Dataflow/DataflowDebugDrawInterface.h"
+#include "Dataflow/DataflowDebugDrawObject.h"
+#include "Dataflow/DataflowRenderingViewMode.h"
+#include "Generators/MinimalBoxMeshGenerator.h"
 
 #include UE_INLINE_GENERATED_CPP_BY_NAME(GenerateCardsClumpsNode)
 
@@ -30,7 +33,7 @@ void FGenerateCardsClumpsNode::Evaluate(UE::Dataflow::FContext& Context, const F
 						{
 							if((OverideSettings.FilterName != NAME_None && FilterSettings->GetFilterName() == OverideSettings.FilterName) || OverideSettings.FilterName == NAME_None)
 							{
-								FilterSettings->TargetNumberOfCards = OverideSettings.NumCards;
+								FilterSettings->TargetNumberOfCards = OverideSettings.NumCards * 3;
 								FilterSettings->MaxNumberOfFlyaways = OverideSettings.NumFlyaways;
 							}
 						}
diff --git a/Engine/Plugins/Experimental/HairCardGenerator/Source/HairCardGeneratorDataflow/Private/HairCardDataflowRendering.cpp b/Engine/Plugins/Experimental/HairCardGenerator/Source/HairCardGeneratorDataflow/Private/HairCardDataflowRendering.cpp
--- a/Engine/Plugins/Experimental/HairCardGenerator/Source/HairCardGeneratorDataflow/Private/HairCardDataflowRendering.cpp
+++ b/Engine/Plugins/Experimental/HairCardGenerator/Source/HairCardGeneratorDataflow/Private/HairCardDataflowRendering.cpp
@@ -217,68 +217,218 @@ namespace UE::CardGen::Private
 
 	Dataflow::FRenderKey FCardsClumpsRenderingCallbacks::RenderKey = { "ClumpsRender", FName("FCardsCollection") };
 	
-	int32 FCardsClumpsRenderingCallbacks::GetGroupAttribute(const UE::Groom::FGroomStrandsFacade& StrandsFacade, FString& GroupAttribute, FString& GroupName) const
+	void FCardsClumpsRenderingCallbacks::Render(GeometryCollection::Facades::FRenderingFacade& RenderCollection, const Dataflow::FGraphRenderingState& State)
 	{
-		GroupAttribute = FGenerateCardsClumpsNode::CurveClumpIndicesAttribute.ToString();
-		GroupAttribute.AppendInt(GGroomDataflowCardsLod);
-		GroupName = FString("Clump");
-
-		const FManagedArrayCollection& GroomCollection = StrandsFacade.GetManagedArrayCollection();
-		if (GroomCollection.HasAttribute(FName(GroupAttribute), UE::Groom::FGroomStrandsFacade::CurvesGroup))
+		if (State.GetRenderOutputs().Num())
 		{
-			const TManagedArray<int32>& ClumpIndices = GroomCollection.GetAttribute<int32>(FName(GroupAttribute), UE::Groom::FGroomStrandsFacade::CurvesGroup);
+			const FManagedArrayCollection Default;
+			checkf(State.GetRenderOutputs().Num() == 1, TEXT("Expected FGraphRenderingState object to have one render output"));
+			const FName PrimaryOutput = State.GetRenderOutputs()[0];
+			const FManagedArrayCollection& Collection = State.GetValue<FManagedArrayCollection>(PrimaryOutput, Default);
 
-			int32 NumClumps = INDEX_NONE;
-			for(int32 ElementIndex = 0, NumElements = GroomCollection.NumElements(UE::Groom::FGroomStrandsFacade::CurvesGroup); ElementIndex < NumElements; ++ElementIndex)
-			{	
-				NumClumps = FMath::Max(NumClumps, ClumpIndices[ElementIndex]);
-			}
-			++NumClumps;
-			return NumClumps;
-		}
-		return 0;
-	}
-	
-	void FCardsClumpsRenderingCallbacks::ComputeVertexColors(const UE::Groom::FGroomStrandsFacade& StrandsFacade, TArray<FLinearColor>& VertexColors) const
-	{
-		const FManagedArrayCollection& GroomCollection = StrandsFacade.GetManagedArrayCollection();
+			FString ClumpIndicesLOD = FGenerateCardsClumpsNode::CurveClumpIndicesAttribute.ToString();
+			ClumpIndicesLOD.AppendInt(GGroomDataflowCardsLod);
+			
+			FString NumClumpsLOD = FGenerateCardsClumpsNode::ObjectNumClumpsAttribute.ToString();
+			NumClumpsLOD.AppendInt(GGroomDataflowCardsLod);
 
-		FString ClumpIndicesLOD = FGenerateCardsClumpsNode::CurveClumpIndicesAttribute.ToString();
-		ClumpIndicesLOD.AppendInt(GGroomDataflowCardsLod);
+			if (Collection.HasAttribute(FName(ClumpIndicesLOD), UE::Groom::FGroomStrandsFacade::CurvesGroup) &&
+				Collection.HasAttribute(FName(NumClumpsLOD), UE::Groom::FGroomStrandsFacade::ObjectsGroup))
+			{
+				if (State.GetViewMode().GetName() == Dataflow::FDataflowConstruction3DViewMode::Name)
+				{
+					const UE::Groom::FGroomStrandsFacade StrandsFacade(Collection);
+					if (StrandsFacade.IsValid())
+					{
+						const TManagedArray<int32>& CurveClumpIndices = Collection.GetAttribute<int32>(FName(ClumpIndicesLOD), UE::Groom::FGroomStrandsFacade::CurvesGroup);
+						const TManagedArray<int32>& ObjectNumClumps = Collection.GetAttribute<int32>(FName(NumClumpsLOD), UE::Groom::FGroomStrandsFacade::ObjectsGroup);
 
-		FString NumClumpsLOD = FGenerateCardsClumpsNode::ObjectNumClumpsAttribute.ToString();
-		NumClumpsLOD.AppendInt(GGroomDataflowCardsLod);
-		
-		if(GroomCollection.HasAttribute(FName(ClumpIndicesLOD), UE::Groom::FGroomStrandsFacade::CurvesGroup) &&
-		   GroomCollection.HasAttribute(FName(NumClumpsLOD), UE::Groom::FGroomStrandsFacade::ObjectsGroup))
-		{
-			const TManagedArray<int32>& ClumpIndices = GroomCollection.GetAttribute<int32>(FName(ClumpIndicesLOD), UE::Groom::FGroomStrandsFacade::CurvesGroup);
-			const TManagedArray<int32>& NumClumps = GroomCollection.GetAttribute<int32>(FName(NumClumpsLOD), UE::Groom::FGroomStrandsFacade::ObjectsGroup);
+						int32 NumClumps = 0;
+						for(int32 CurveIndex = 0, NumCurves = CurveClumpIndices.Num(); CurveIndex < NumCurves; ++CurveIndex)
+						{
+							NumClumps = FMath::Max(NumClumps, CurveClumpIndices[CurveIndex]);
+						}
+						++NumClumps;
+
+						TArray<TArray<int32>> ClumpCurveIndices;
+						TArray<TArray<TArray<FVector>>> ClumpMeshPoints;
+						TArray<TArray<FVector>> ClumpMeshOrigins;
+						TArray<TArray<FQuat>> ClumpMeshRotations;
+						TArray<TArray<FVector>> ClumpMeshExtentsX;
+						TArray<TArray<FVector>> ClumpMeshExtentsY;
+
+						ClumpCurveIndices.SetNum(NumClumps);
+						ClumpMeshPoints.SetNum(NumClumps);
+						ClumpMeshOrigins.SetNum(NumClumps);
+						ClumpMeshRotations.SetNum(NumClumps);
+						ClumpMeshExtentsX.SetNum(NumClumps);
+						ClumpMeshExtentsY.SetNum(NumClumps);
+
+						// Fill the clump curve indices
+						for(int32 CurveIndex = 0, NumCurves = CurveClumpIndices.Num(); CurveIndex < NumCurves; ++CurveIndex)
+						{
+							if (ClumpCurveIndices.IsValidIndex(CurveClumpIndices[CurveIndex]))
+							{
+								ClumpCurveIndices[CurveClumpIndices[CurveIndex]].Add(CurveIndex);
+							}
+						}
 
-			TArray<FLinearColor> ClumpsColors;
-			ClumpsColors.Init(FLinearColor::Black, NumClumps[0]);
+						static constexpr const int32 NumMeshes = 10;
+						static constexpr const int32 NumCorners = 8;
+						static constexpr const int32 NumTriangles = 12;
+						static constexpr const float MinExtent = 0.5f;
 
-			for(int32 ClumpIndex = 0; ClumpIndex < NumClumps[0]; ++ClumpIndex)
-			{
-				ClumpsColors[ClumpIndex] = FLinearColor::IntToDistinctColor(ClumpIndex, 0.75f, 1.0f, 90.0f);
-			}
+						for(int32 ClumpIndex = 0; ClumpIndex < NumClumps; ++ClumpIndex)
+						{
+							// Fill the clump points for eaqch meshes
+							ClumpMeshPoints[ClumpIndex].SetNum(NumMeshes+1);
+							for (const int32& CurveIndex : ClumpCurveIndices[ClumpIndex])
+							{
+								const float PointStart = (CurveIndex > 0) ? StrandsFacade.GetCurvePointOffsets()[CurveIndex - 1] : 0;
+								const float PointEnd = StrandsFacade.GetCurvePointOffsets()[CurveIndex]-1;
 
-			TArray<FLinearColor> StrandsColors;
-			StrandsColors.Init(FLinearColor::Black, StrandsFacade.GetNumCurves());
-			for(int32 CurveIndex = 0, CurveEnd = StrandsFacade.GetNumCurves(); CurveIndex < CurveEnd; ++CurveIndex)
-			{
-				if(ClumpIndices[CurveIndex] < ClumpsColors.Num())
-				{
-					StrandsColors[CurveIndex] = ClumpsColors[ClumpIndices[CurveIndex]];
-				}
-			}
-			
-			for(int32 PointIndex = 0, PointEnd = StrandsFacade.GetNumPoints(); PointIndex < PointEnd; ++PointIndex)
-			{
-				const int32 CurveIndex = StrandsFacade.GetPointCurveIndices()[PointIndex];
+								ClumpMeshPoints[ClumpIndex][0].Add(FVector(StrandsFacade.GetPointRestPositions()[PointStart]));
+								ClumpMeshPoints[ClumpIndex][NumMeshes].Add(FVector(StrandsFacade.GetPointRestPositions()[PointEnd]));
+
+								for(int32 MeshIndex = 1; MeshIndex < NumMeshes; ++MeshIndex)
+								{
+									const float PointParam = PointStart + (PointEnd - PointStart) * MeshIndex / NumMeshes;
+									const int32 ParamIndex = FMath::FloorToInt(PointParam);
+									const float ParamAlpha = PointParam - ParamIndex;
+
+									ClumpMeshPoints[ClumpIndex][MeshIndex].Add(FVector(StrandsFacade.GetPointRestPositions()[ParamIndex] * (1.0 - ParamAlpha) + StrandsFacade.GetPointRestPositions()[ParamIndex+1] * ParamAlpha));
+								}
+							}
+
+							// Compute the clump mesh origins
+							ClumpMeshOrigins[ClumpIndex].Init(FVector::ZeroVector, NumMeshes + 1);
+							for (int32 MeshIndex = 0; MeshIndex < NumMeshes+1; ++MeshIndex)
+							{
+								if(!ClumpMeshPoints[ClumpIndex][MeshIndex].IsEmpty())
+								{ 
+									for (const FVector& MeshPoint : ClumpMeshPoints[ClumpIndex][MeshIndex])
+									{
+										ClumpMeshOrigins[ClumpIndex][MeshIndex] += MeshPoint;
+									}
+									ClumpMeshOrigins[ClumpIndex][MeshIndex] /= ClumpMeshPoints[ClumpIndex][MeshIndex].Num();
+								}
+							}
+
+							FVector NextTangent = FVector(0,0,1), PrevTangent = FVector(0, 0, 1);
+							FQuat MeshOrientation = FQuat::Identity;
+
+							// Compute the clump mesh rotations
+							ClumpMeshRotations[ClumpIndex].Init(FQuat::Identity, NumMeshes);
+							for (int32 MeshIndex = 0; MeshIndex < NumMeshes; ++MeshIndex)
+							{
+								PrevTangent = NextTangent;
+								NextTangent = (ClumpMeshOrigins[ClumpIndex][MeshIndex+1]-ClumpMeshOrigins[ClumpIndex][MeshIndex]).GetSafeNormal();
+
+								MeshOrientation = (FQuat::FindBetweenNormals(PrevTangent, NextTangent) * MeshOrientation).GetNormalized();
+								ClumpMeshRotations[ClumpIndex][MeshIndex] = MeshOrientation;
+							}
+
+							// Compute the clump mesh extents
+							ClumpMeshExtentsX[ClumpIndex].Init(FVector::ZeroVector, NumMeshes + 1);
+							ClumpMeshExtentsY[ClumpIndex].Init(FVector::ZeroVector, NumMeshes + 1);
 
-				VertexColors[2*PointIndex] = StrandsColors[CurveIndex];
-				VertexColors[2*PointIndex+1] = StrandsColors[CurveIndex];
+							float MeshExtent = MinExtent;
+							for (int32 MeshIndex = 0; MeshIndex < NumMeshes+1; ++MeshIndex)
+							{
+								MeshOrientation = (MeshIndex == NumMeshes) ? ClumpMeshRotations[ClumpIndex][MeshIndex-1] : ClumpMeshRotations[ClumpIndex][MeshIndex];
+
+								const FVector MeshNormalX = MeshOrientation.GetAxisX();
+								const FVector MeshNormalY = MeshOrientation.GetAxisY();
+
+								if(MeshIndex == 0)
+								{ 
+									FVector2d MeshExtents(MinExtent, MinExtent);
+									for (const FVector& MeshPoint : ClumpMeshPoints[ClumpIndex][MeshIndex])
+									{
+										const FVector LocalPoint = MeshPoint - ClumpMeshOrigins[ClumpIndex][MeshIndex];
+
+										MeshExtents[0] = FMath::Max(MeshExtents[0], FMath::Abs(LocalPoint.Dot(MeshNormalX)));
+										MeshExtents[1] = FMath::Max(MeshExtents[1], FMath::Abs(LocalPoint.Dot(MeshNormalY)));
+									}
+									MeshExtent = FMath::Min(MeshExtents[0], MeshExtents[1]);
+								}
+								ClumpMeshExtentsX[ClumpIndex][MeshIndex] = MeshNormalX * MeshExtent;
+								ClumpMeshExtentsY[ClumpIndex][MeshIndex] = MeshNormalY * MeshExtent;
+							}
+
+							TArray<FVector3f> VertexLocalPositions;
+							TArray<FVector3f> VertexLocalNormals;
+							TArray<FIntVector> FaceLocalVertices;
+
+							VertexLocalNormals.Init(FVector3f::ZeroVector, NumMeshes * NumCorners);
+							VertexLocalPositions.Reserve(NumMeshes * NumCorners);
+							FaceLocalVertices.Reserve(NumMeshes * NumTriangles);
+
+							int32 VertexOffset = 0;
+
+							auto AddMeshFace = [&VertexOffset, &FaceLocalVertices, &VertexLocalNormals, &VertexLocalPositions](const int32 VertexIndexA, const int32 VertexIndexB, const int32 VertexIndexC)
+							{
+								FIntVector& FaceIndices = FaceLocalVertices.Add_GetRef(FIntVector(VertexOffset + VertexIndexA, VertexOffset + VertexIndexB, VertexOffset + VertexIndexC));
+								const FVector3f FaceNormal = ((VertexLocalPositions[FaceIndices[2]] - VertexLocalPositions[FaceIndices[0]]).Cross(VertexLocalPositions[FaceIndices[1]] - VertexLocalPositions[FaceIndices[0]])).GetSafeNormal();
+
+								VertexLocalNormals[FaceIndices[0]] += FaceNormal;
+								VertexLocalNormals[FaceIndices[1]] += FaceNormal;
+								VertexLocalNormals[FaceIndices[2]] += FaceNormal;
+							};
+							for (int32 MeshIndex = 0; MeshIndex < NumMeshes; ++MeshIndex)
+							{
+								// Add vertex positions
+								FVector3f MeshOrigin = FVector3f(ClumpMeshOrigins[ClumpIndex][MeshIndex]);
+								FVector3f MeshExtentX = FVector3f(ClumpMeshExtentsX[ClumpIndex][MeshIndex]);
+								FVector3f MeshExtentY = FVector3f(ClumpMeshExtentsY[ClumpIndex][MeshIndex]);
+
+								VertexLocalPositions.Add(MeshOrigin - MeshExtentX - MeshExtentY);
+								VertexLocalPositions.Add(MeshOrigin + MeshExtentX - MeshExtentY);
+								VertexLocalPositions.Add(MeshOrigin + MeshExtentX + MeshExtentY);
+								VertexLocalPositions.Add(MeshOrigin - MeshExtentX + MeshExtentY);
+
+								MeshOrigin = FVector3f(ClumpMeshOrigins[ClumpIndex][MeshIndex+1]);
+								MeshExtentX = FVector3f(ClumpMeshExtentsX[ClumpIndex][MeshIndex+1]);
+								MeshExtentY = FVector3f(ClumpMeshExtentsY[ClumpIndex][MeshIndex+1]);
+
+								VertexLocalPositions.Add(MeshOrigin - MeshExtentX - MeshExtentY);
+								VertexLocalPositions.Add(MeshOrigin + MeshExtentX - MeshExtentY);
+								VertexLocalPositions.Add(MeshOrigin + MeshExtentX + MeshExtentY);
+								VertexLocalPositions.Add(MeshOrigin - MeshExtentX + MeshExtentY);
+
+								// Add box faces and update normals
+								AddMeshFace(0, 1, 2);
+								AddMeshFace(2, 3, 0);
+								AddMeshFace(7, 6, 5);
+								AddMeshFace(5, 4, 7);
+								AddMeshFace(3, 2, 6);
+								AddMeshFace(6, 7, 3);
+								AddMeshFace(0, 3, 7);
+								AddMeshFace(7, 4, 0);
+								AddMeshFace(2, 1, 5);
+								AddMeshFace(5, 6, 2);
+								AddMeshFace(1, 0, 4);
+								AddMeshFace(4, 5, 1);
+
+								for(int32 VertexIndex = 0; VertexIndex < NumCorners; ++VertexIndex)
+								{
+									VertexLocalNormals[VertexIndex + VertexOffset].Normalize();
+								}
+								VertexOffset += NumCorners;
+							}
+
+							const FLinearColor GroupColor = FLinearColor::IntToDistinctColor(ClumpIndex, 0.75f, 1.0f, 90.0f);
+
+							TArray<FLinearColor> VertexLocalColors;
+							VertexLocalColors.Init(GroupColor, VertexLocalPositions.Num());
+
+							const FString GeometryName = TEXT("Groom_Clump_") + FString::FromInt(ClumpIndex);
+							const int32 GeometryIndex = RenderCollection.StartGeometryGroup(GeometryName);
+							RenderCollection.AddSurface(MoveTemp(VertexLocalPositions), MoveTemp(FaceLocalVertices), MoveTemp(VertexLocalNormals), MoveTemp(VertexLocalColors));
+							RenderCollection.EndGeometryGroup(GeometryIndex);
+						}
+					}
+				}
 			}
 		}
 	}
diff --git a/Engine/Plugins/Experimental/HairCardGenerator/Source/HairCardGeneratorDataflow/Public/GenerateCardsClumpsNode.h b/Engine/Plugins/Experimental/HairCardGenerator/Source/HairCardGeneratorDataflow/Public/GenerateCardsClumpsNode.h
--- a/Engine/Plugins/Experimental/HairCardGenerator/Source/HairCardGeneratorDataflow/Public/GenerateCardsClumpsNode.h
+++ b/Engine/Plugins/Experimental/HairCardGenerator/Source/HairCardGeneratorDataflow/Public/GenerateCardsClumpsNode.h
@@ -21,8 +21,8 @@ struct FCardsClumpsSettings
 	FName FilterName = NAME_None;
 
 	// Total number of cards to generate for this LOD settings
-	UPROPERTY(EditAnywhere, Category="Groom Cards", meta=(ClampMin="1", ClampMax="100000"))
-	int32 NumCards = 200;
+	UPROPERTY(EditAnywhere, Category="Groom Cards", meta=(ClampMin="1", ClampMax="100000", DisplayName = "Num Clumps"))
+	int32 NumCards = 50;
 
 	// Maximum number of cards to assign to flyaway strands
 	UPROPERTY(EditAnywhere, Category="Groom Cards", meta=(ClampMin="0", ClampMax="1000"))
diff --git a/Engine/Plugins/Experimental/HairCardGenerator/Source/HairCardGeneratorDataflow/Public/HairCardDataflowRendering.h b/Engine/Plugins/Experimental/HairCardGenerator/Source/HairCardGeneratorDataflow/Public/HairCardDataflowRendering.h
--- a/Engine/Plugins/Experimental/HairCardGenerator/Source/HairCardGeneratorDataflow/Public/HairCardDataflowRendering.h
+++ b/Engine/Plugins/Experimental/HairCardGenerator/Source/HairCardGeneratorDataflow/Public/HairCardDataflowRendering.h
@@ -28,23 +28,25 @@ namespace UE::CardGen::Private
 		virtual void Render(GeometryCollection::Facades::FRenderingFacade& RenderCollection, const Dataflow::FGraphRenderingState& State) override;
 		virtual void ComputeVertexColors(const FManagedArrayCollection& Collection, const int32 CardIndex, TArray<FLinearColor>& VertexColors) const {}
 	};
-
-	class FCardsClumpsRenderingCallbacks : public UE::Groom::FGroomStrandsRenderingCallbacks
+	
+	class FCardsClumpsRenderingCallbacks : public Dataflow::FRenderingFactory::ICallbackInterface
 	{
-	public :
-		static Dataflow::FRenderKey RenderKey;
+		public :
+			static Dataflow::FRenderKey RenderKey;
 		
-	private:
+	protected:
 
 		virtual Dataflow::FRenderKey GetRenderKey() const override
 		{
 			return RenderKey;
 		}
-		
-		virtual void ComputeVertexColors(const UE::Groom::FGroomStrandsFacade& StrandsFacade, TArray<FLinearColor>& VertexColors) const override;
-		
-		/** Get the group attribute used to create each geometry groups */
-		virtual int32 GetGroupAttribute(const UE::Groom::FGroomStrandsFacade& StrandsFacade, FString& GroupAttribute, FString& GroupName) const override;
+
+		virtual bool CanRender(const Dataflow::IDataflowConstructionViewMode& ViewMode) const override
+		{
+			return (ViewMode.GetName() == Dataflow::FDataflowConstruction3DViewMode::Name);
+		}
+
+		virtual void Render(GeometryCollection::Facades::FRenderingFacade& RenderCollection, const Dataflow::FGraphRenderingState& State) override;
 	};
 
 	class FCardsTextureRenderingCallbacks : public FCardsGeometryRenderingCallbacks
