Display card group names to be able to create filter settings

diff --git a/Engine/Plugins/Experimental/HairCardGenerator/Source/HairCardGeneratorDataflow/Private/BuildCardsSettingsNode.cpp b/Engine/Plugins/Experimental/HairCardGenerator/Source/HairCardGeneratorDataflow/Private/BuildCardsSettingsNode.cpp
--- a/Engine/Plugins/Experimental/HairCardGenerator/Source/HairCardGeneratorDataflow/Private/BuildCardsSettingsNode.cpp
+++ b/Engine/Plugins/Experimental/HairCardGenerator/Source/HairCardGeneratorDataflow/Private/BuildCardsSettingsNode.cpp
@@ -4,6 +4,7 @@
 
 #include "GroomCollectionFacades.h"
 #include "HairCardGeneratorEditorModule.h"
+#include "Dataflow/DataflowDebugDrawInterface.h"
 
 #include UE_INLINE_GENERATED_CPP_BY_NAME(BuildCardsSettingsNode)
 
@@ -31,7 +32,29 @@ namespace UE::CardGen::Private
 			}
 		}
 	}
+	
+	static FString DisplayCardsGroups(const UGroomAsset* GroomAsset)
+	{
+		TArray<FName> GroupNames;
+		FHairCardGeneratorUtils::BuildCardsGroups(GroomAsset, GroupNames);
+		
+		bool bFirstGroup = true;
+		FString CardsGroups;
+		CardsGroups += "[";
+		for(const FName& GroupName : GroupNames)
+		{
+			if(!bFirstGroup)
+			{
+				CardsGroups +=  ",";
+			}
+			CardsGroups += GroupName.ToString();
+			bFirstGroup = false;
+		}
+		CardsGroups += "]";
+		return CardsGroups;
+	}
 }
+
 void FBuildCardsSettingsNode::Evaluate(UE::Dataflow::FContext& Context, const FDataflowOutput* Out) const
 {
 	if (Out->IsA<FManagedArrayCollection>(&Collection))
@@ -87,7 +110,36 @@ void FBuildCardsSettingsNode::Evaluate(UE::Dataflow::FContext& Context, const FD
 				}
 			}
 		}
+		else
+		{
+			SetError(Context, &CardsSettings, TEXT("No valid groom cards groups. Please add a new group in the groom hair cards tab"));
+		}
 		SetValue(Context, MoveTemp(OutputSettings), &CardsSettings);
 	}
 }
 
+
+#if WITH_EDITOR
+void FBuildCardsSettingsNode::DebugDraw(UE::Dataflow::FContext& Context, IDataflowDebugDrawInterface& DataflowRenderingInterface, const FDebugDrawParameters& DebugDrawParameters) const
+{
+	if ((DebugDrawParameters.bNodeIsSelected || DebugDrawParameters.bNodeIsPinned))
+	{
+		UGroomAsset* LocalAsset = const_cast<UGroomAsset*>(GroomAsset.Get());
+		if (!LocalAsset)
+		{
+			if (const UE::Dataflow::FEngineContext* EngineContext = Context.AsType<UE::Dataflow::FEngineContext>())
+			{
+				LocalAsset = Cast<UGroomAsset>(EngineContext->Owner);
+			}
+		}
+		if(LocalAsset)
+		{
+			FString CardsInfos = FString("Groom Asset Name : ") + LocalAsset->GetName() + FString("\n");
+			CardsInfos += FString("Cards Groups Names : ") + UE::CardGen::Private::DisplayCardsGroups(LocalAsset) + FString("\n");
+			CardsInfos +=  FString("Num Cards Generations : ") + FString::FromInt(LocalAsset->GetHairGroupsCards().Num());
+			DataflowRenderingInterface.DrawOverlayText(CardsInfos);
+		}
+	}
+}
+#endif
+
diff --git a/Engine/Plugins/Experimental/HairCardGenerator/Source/HairCardGeneratorDataflow/Private/HairCardDataflowRendering.cpp b/Engine/Plugins/Experimental/HairCardGenerator/Source/HairCardGeneratorDataflow/Private/HairCardDataflowRendering.cpp
--- a/Engine/Plugins/Experimental/HairCardGenerator/Source/HairCardGeneratorDataflow/Private/HairCardDataflowRendering.cpp
+++ b/Engine/Plugins/Experimental/HairCardGenerator/Source/HairCardGeneratorDataflow/Private/HairCardDataflowRendering.cpp
@@ -241,7 +241,6 @@ namespace UE::CardGen::Private
 					if (StrandsFacade.IsValid())
 					{
 						const TManagedArray<int32>& CurveClumpIndices = Collection.GetAttribute<int32>(FName(ClumpIndicesLOD), UE::Groom::FGroomStrandsFacade::CurvesGroup);
-						const TManagedArray<int32>& ObjectNumClumps = Collection.GetAttribute<int32>(FName(NumClumpsLOD), UE::Groom::FGroomStrandsFacade::ObjectsGroup);
 
 						int32 NumClumps = 0;
 						for(int32 CurveIndex = 0, NumCurves = CurveClumpIndices.Num(); CurveIndex < NumCurves; ++CurveIndex)
@@ -267,7 +266,7 @@ namespace UE::CardGen::Private
 						// Fill the clump curve indices
 						for(int32 CurveIndex = 0, NumCurves = CurveClumpIndices.Num(); CurveIndex < NumCurves; ++CurveIndex)
 						{
-							if (ClumpCurveIndices.IsValidIndex(CurveClumpIndices[CurveIndex]))
+							if ((CurveClumpIndices[CurveIndex] >= 0) && ClumpCurveIndices.IsValidIndex(CurveClumpIndices[CurveIndex]))
 							{
 								ClumpCurveIndices[CurveClumpIndices[CurveIndex]].Add(CurveIndex);
 							}
@@ -284,19 +283,28 @@ namespace UE::CardGen::Private
 							ClumpMeshPoints[ClumpIndex].SetNum(NumMeshes+1);
 							for (const int32& CurveIndex : ClumpCurveIndices[ClumpIndex])
 							{
-								const float PointStart = (CurveIndex > 0) ? StrandsFacade.GetCurvePointOffsets()[CurveIndex - 1] : 0;
-								const float PointEnd = StrandsFacade.GetCurvePointOffsets()[CurveIndex]-1;
-
-								ClumpMeshPoints[ClumpIndex][0].Add(FVector(StrandsFacade.GetPointRestPositions()[PointStart]));
-								ClumpMeshPoints[ClumpIndex][NumMeshes].Add(FVector(StrandsFacade.GetPointRestPositions()[PointEnd]));
-
-								for(int32 MeshIndex = 1; MeshIndex < NumMeshes; ++MeshIndex)
+								if(StrandsFacade.GetCurvePointOffsets().IsValidIndex(CurveIndex) && StrandsFacade.GetCurvePointOffsets().IsValidIndex(CurveIndex - 1))
 								{
-									const float PointParam = PointStart + (PointEnd - PointStart) * MeshIndex / NumMeshes;
-									const int32 ParamIndex = FMath::FloorToInt(PointParam);
-									const float ParamAlpha = PointParam - ParamIndex;
+									const float PointStart = (CurveIndex > 0) ? StrandsFacade.GetCurvePointOffsets()[CurveIndex - 1] : 0;
+									const float PointEnd = StrandsFacade.GetCurvePointOffsets()[CurveIndex]-1;
 
-									ClumpMeshPoints[ClumpIndex][MeshIndex].Add(FVector(StrandsFacade.GetPointRestPositions()[ParamIndex] * (1.0 - ParamAlpha) + StrandsFacade.GetPointRestPositions()[ParamIndex+1] * ParamAlpha));
+									if(StrandsFacade.GetPointRestPositions().IsValidIndex(PointStart) && StrandsFacade.GetPointRestPositions().IsValidIndex(PointEnd))
+									{
+										ClumpMeshPoints[ClumpIndex][0].Add(FVector(StrandsFacade.GetPointRestPositions()[PointStart]));
+										ClumpMeshPoints[ClumpIndex][NumMeshes].Add(FVector(StrandsFacade.GetPointRestPositions()[PointEnd]));
+
+										for(int32 MeshIndex = 1; MeshIndex < NumMeshes; ++MeshIndex)
+										{
+											const float PointParam = PointStart + (PointEnd - PointStart) * MeshIndex / NumMeshes;
+											const int32 ParamIndex = FMath::FloorToInt(PointParam);
+											const float ParamAlpha = PointParam - ParamIndex;
+
+											if(StrandsFacade.GetPointRestPositions().IsValidIndex(ParamIndex) && StrandsFacade.GetPointRestPositions().IsValidIndex(ParamIndex+1))
+											{
+												ClumpMeshPoints[ClumpIndex][MeshIndex].Add(FVector(StrandsFacade.GetPointRestPositions()[ParamIndex] * (1.0 - ParamAlpha) + StrandsFacade.GetPointRestPositions()[ParamIndex+1] * ParamAlpha));
+											}
+										}
+									}
 								}
 							}
 
diff --git a/Engine/Plugins/Experimental/HairCardGenerator/Source/HairCardGeneratorDataflow/Public/BuildCardsSettingsNode.h b/Engine/Plugins/Experimental/HairCardGenerator/Source/HairCardGeneratorDataflow/Public/BuildCardsSettingsNode.h
--- a/Engine/Plugins/Experimental/HairCardGenerator/Source/HairCardGeneratorDataflow/Public/BuildCardsSettingsNode.h
+++ b/Engine/Plugins/Experimental/HairCardGenerator/Source/HairCardGeneratorDataflow/Public/BuildCardsSettingsNode.h
@@ -75,6 +75,11 @@ public:
 
 	//~ Begin FDataflowNode interface
 	virtual void Evaluate(UE::Dataflow::FContext& Context, const FDataflowOutput* Out) const override;
+#if WITH_EDITOR
+	virtual bool CanDebugDraw() const override { return true; }
+	virtual bool CanDebugDrawViewMode(const FName& ViewModeName) const override {return true;}
+	virtual void DebugDraw(UE::Dataflow::FContext& Context, IDataflowDebugDrawInterface& DataflowRenderingInterface, const FDebugDrawParameters& DebugDrawParameters) const override;
+#endif
 	//~ End FDataflowNode interface
 
 	/** Groom asset to build the cards settings from */
@@ -89,7 +94,7 @@ public:
 	UPROPERTY(meta = (DataflowOutput));
 	TArray<FGroomCardsSettings> CardsSettings;
 	
-	/** List of filter setings to override */
+	/** List of filter settings to override */
 	UPROPERTY(EditAnywhere, Category = "Groom Cards")
 	TArray<FGroomFilterSettings> FilterSettings;
 };
diff --git a/Engine/Plugins/Experimental/HairCardGenerator/Source/HairCardGeneratorEditor/Private/HairCardGeneratorEditorModule.cpp b/Engine/Plugins/Experimental/HairCardGenerator/Source/HairCardGeneratorEditor/Private/HairCardGeneratorEditorModule.cpp
--- a/Engine/Plugins/Experimental/HairCardGenerator/Source/HairCardGeneratorEditor/Private/HairCardGeneratorEditorModule.cpp
+++ b/Engine/Plugins/Experimental/HairCardGenerator/Source/HairCardGeneratorEditor/Private/HairCardGeneratorEditorModule.cpp
@@ -409,6 +409,11 @@ static bool GenerateHairCardGenGroomData(const int PointsPerStrand, const FHairD
 	return true;
 }
 
+void FHairCardGeneratorUtils::BuildCardsGroups(TObjectPtr<const UGroomAsset> GroomAsset, TArray<FName>& GroupNames)
+{
+	HairCardSettings_Helpers::GetStrandCardGroupInfo(GroomAsset, GroupNames);
+}
+
 bool FHairCardGeneratorUtils::LoadGroomStrands(const UGroomAsset* GroomAsset, const FBuilderFunction& BuilderFunction)
 {
 	TObjectPtr<UHairCardGenControllerBase> CardGenController = FHairCardGenController::GetHairCardGenerator();
diff --git a/Engine/Plugins/Experimental/HairCardGenerator/Source/HairCardGeneratorEditor/Private/HairCardGeneratorPluginSettings.cpp b/Engine/Plugins/Experimental/HairCardGenerator/Source/HairCardGeneratorEditor/Private/HairCardGeneratorPluginSettings.cpp
--- a/Engine/Plugins/Experimental/HairCardGenerator/Source/HairCardGeneratorEditor/Private/HairCardGeneratorPluginSettings.cpp
+++ b/Engine/Plugins/Experimental/HairCardGenerator/Source/HairCardGeneratorEditor/Private/HairCardGeneratorPluginSettings.cpp
@@ -100,8 +100,6 @@ namespace HairCardSettings_Helpers
 	void ResetSettingsFromTemplate(TObjectPtr<UObject> SettingsObject, TObjectPtr<const UObject> TemplateObject);
 	uint8 PipelineSettingsDiff(TObjectPtr<const UObject> SettingsObject, TObjectPtr<UObject> DiffObject);
 	bool IsWriteableSettingsProperty(FProperty* Property);
-
-    void GetStrandCardGroupInfo(TObjectPtr<const UGroomAsset> GroomAsset, TArray<FName>& OutGrpIds);
 };
 
 
diff --git a/Engine/Plugins/Experimental/HairCardGenerator/Source/HairCardGeneratorEditor/Private/HairCardGeneratorPluginSettings.h b/Engine/Plugins/Experimental/HairCardGenerator/Source/HairCardGeneratorEditor/Private/HairCardGeneratorPluginSettings.h
--- a/Engine/Plugins/Experimental/HairCardGenerator/Source/HairCardGeneratorEditor/Private/HairCardGeneratorPluginSettings.h
+++ b/Engine/Plugins/Experimental/HairCardGenerator/Source/HairCardGeneratorEditor/Private/HairCardGeneratorPluginSettings.h
@@ -10,6 +10,11 @@
 
 class FJsonObject;
 
+namespace HairCardSettings_Helpers
+{
+	void GetStrandCardGroupInfo(TObjectPtr<const UGroomAsset> GroomAsset, TArray<FName>& OutGrpIds);
+};
+
 UENUM(BlueprintType)
 enum class EHairCardAtlasSize : uint8
 {
diff --git a/Engine/Plugins/Experimental/HairCardGenerator/Source/HairCardGeneratorEditor/Public/HairCardGeneratorEditorModule.h b/Engine/Plugins/Experimental/HairCardGenerator/Source/HairCardGeneratorEditor/Public/HairCardGeneratorEditorModule.h
--- a/Engine/Plugins/Experimental/HairCardGenerator/Source/HairCardGeneratorEditor/Public/HairCardGeneratorEditorModule.h
+++ b/Engine/Plugins/Experimental/HairCardGenerator/Source/HairCardGeneratorEditor/Public/HairCardGeneratorEditorModule.h
@@ -43,6 +43,9 @@ struct HAIRCARDGENERATOREDITOR_API FHairCardGeneratorUtils
 
 	/** Load the groom strands */
 	static bool LoadGroomStrands(const UGroomAsset* GroomAsset, const FBuilderFunction& BuilderFunction);
+	
+	/** Build the cards group names */
+	static void BuildCardsGroups(TObjectPtr<const UGroomAsset> GroomAsset, TArray<FName>& GroupNames);
 };
 
 /** FHairCardGeneratorEditorModule  */
