HairCardGenerator: Update python gl dependencies for better cross-platform support.

diff --git a/Engine/Plugins/Experimental/HairCardGenerator/HairCardGenerator.uplugin b/Engine/Plugins/Experimental/HairCardGenerator/HairCardGenerator.uplugin
--- a/Engine/Plugins/Experimental/HairCardGenerator/HairCardGenerator.uplugin
+++ b/Engine/Plugins/Experimental/HairCardGenerator/HairCardGenerator.uplugin
@@ -22,6 +22,7 @@
 			"LoadingPhase": "Default",
 			"PlatformAllowList": [
 				"Win64",
+				"Mac",
 				"Linux"
 			]
 		},
@@ -31,6 +32,7 @@
 			"LoadingPhase": "Default",
 			"PlatformAllowList": [
 				"Win64",
+				"Mac",
 				"Linux"
 			]
 		}
@@ -49,21 +51,19 @@
 			"Enabled": true
 		}
 	],
-	"PythonRequirements":
-	[
+	"PythonRequirements": [
 		{
 			"Platform": "All",
-			"Requirements":
-			[
+			"Requirements": [
 				"addict==2.4.0 --hash=sha256:249bb56bbfd3cdc2a004ea0ff4c2b6ddc84d53bc2194761636eb314d5cfa5dfc",
-				"glcontext==2.5.0 --hash=sha256:675091975c2bf8b59e9c41d1daa6684006162c3cf3d9e2a1acda31b35181d9d2 --hash=sha256:798bc74604e306386f858be11aa1fe47c88296ac6fb9b5718a1c5a4cfb24416c --hash=sha256:3e0f3fbbe483f3e671ae04698a2d38234cb9a5682d2edd49d5bce08a32d48ff1 --hash=sha256:7d50c62fae0af1b19daa95571d52a5c56f3f1537483f105b4d092be5eb160c9d --hash=sha256:7a7962f89350966966a4a5628fd4861a81dc4f76014571567a54db4c1573e04d --hash=sha256:25da4a8a8707f88e66d1597c5f03be31b354b6d6186907ad06b4735f981aa25e --hash=sha256:942486de098c035dad9165e4cb6ad6342edb8c98b11ee5cdfe97088c7c6840aa --hash=sha256:bd0f159c01dfaab990fedb9672f4be040ae7fe066afb2ce7413c63afa9475f38 --hash=sha256:194b2657f46310cd86662d946d85162710b43e4abbef800c83a61f44b09352ad --hash=sha256:bf968a827c04d5210afc2a84d825719fa3025fd9def4aea4820e75a607d09ccf --hash=sha256:cc2a13791007d18b71fc9eccd3423be0a3d5c16b8d1ac4410767665a9824cc21 --hash=sha256:5c056bfbc4af86337837fbea0899b1c439673b9e2bf9aaaf78862cb68ccaeb41",
+				"glcontext==3.0.0 --hash=sha256:04921720740438ceea8fb8a38b5665963520c7c8f27bef03df8aeb3ea3cfbfb6 --hash=sha256:12c2abef8efabb8ab7e35d16785968de888ae7349d1b83c765080f35fcd3c6e5 --hash=sha256:1445a03d8795113034e1f9ffa662f795df65ae69ee21b26ed3b1d66100ba3f8c --hash=sha256:3a9e56fa3597cc709cfd0fdf2ae682cda36510a13faac2b3142f401e823b64f4 --hash=sha256:4817f4cd52c7fe5410c92ca12b6712435548918719373882ade76f8f75d80abd --hash=sha256:4a2972f92da9a6fb06860e117de05f5b8adc2e2d827bbc0ccc7acbe7325acd1e --hash=sha256:4b3b12a66f57379566dd4d36899ac265abdbe040f3fc3293f50cd6678a1dcc9b --hash=sha256:7043f59d126feb26896a6419ebfeecc78c07ffefced2a2f59104dd7a2f71ebb5 --hash=sha256:7fa413f4420abff2bbb5aa5770a3e1deffcdc13e0ef2f459b145fa79c36909e7 --hash=sha256:848f870a2bc72a29de7ab6756b9e8f2e6ce052e17873ebc6b3f25129b6e0d58a --hash=sha256:a0484308af75e04b0e56066dc2324a8fb9f1443b76ddb98833439982322b2a39 --hash=sha256:b154c25a57e16dbb073478d0cbe2c0090649d135c4c9f87e753b6181b97ec848 --hash=sha256:cdd81e8b580e43c1fe1c48c0fc3909e6c1f37ee4cfd8c990c03b2df3b463bd37 --hash=sha256:e80bb37ba727bd20c192f2754aea40c437a7665005c1001c10752f91913964e9 --hash=sha256:ef4b4ec35e2b720f4cd250bb92cf6417add445490bf780345596da5a796a0e6f --hash=sha256:fa5a14778d13ecf4a0dd60a7825427bf6ac0383eacb3b8b1a9c23e4976aa0c8c",
 				"joblib==1.2.0 --hash=sha256:091138ed78f800342968c523bdde947e7a305b8594b910a0fea2ab83c3c6d385",
-				"moderngl==5.9.0 --hash=sha256:7e68668aa6ea18b9e13b6394e2a979e8ea9e3f70c85988b43724c0d14ad08e7e --hash=sha256:a80ec2f8aa27248d95c100ddbf58d193aef63cfcfcee8acf260a027bf305328e --hash=sha256:377368f777f87ff096c6277c8f2c5100a9742cf3b7fc444c0359fee7e6a17d19 --hash=sha256:ccb7580dbf4212ee0e87733762186bb41ac1bf782996a9260c100fdf13019452 --hash=sha256:4b7008993303eac0cd666191b006f6f08f681d70cd31e48b792aace914c1907f --hash=sha256:5d5e496b2c735906b9d7107dd16cd96b31c4af95735235d05d60d40ae33aa464 --hash=sha256:c26beaf0e93317b1b9d2b4ae3c178a0ae0bb4c1f9b3975d2edc27c38fc4431c9 --hash=sha256:34e73e2e6074dc9067be89de051134d10e84bc85318e4d599ad89b7a0f7fee2e --hash=sha256:7886633683cc2ca00be643778a59d42dd2711eb3ce7eb8cb43d8d494419d752f --hash=sha256:8387b81fe3c230cd954abce08ac91186918e6d0b9681649b118b140209f5fed7 --hash=sha256:e118bebc7f6b060aef0d0eb15323676d28ce93b19f4383cf918399700a2b7e7d --hash=sha256:1ebcc3615d3c20be989005d172a3a74a07df36fca1c5f657afd09d2811982014",
+				"moderngl==5.12.0 --hash=sha256:0c210e8d52a60025f6586ca015c39feb1e57e6dc792c3ff44800f6493a541b1a --hash=sha256:105cdee2ee29a7e6ebbc76cf178941503656a185c6945933bc7ef395ba8e65a7 --hash=sha256:18deb8bebd0a4277d92c76dbedf8e4b4b68bf0a8a878404c6b26aed750890d3b --hash=sha256:2750547707c1ec3790dfbeb9c90fb808672ff13f61cac392c706ba09fda10db0 --hash=sha256:28b9eb3574ffc6e303173ca0a419b63d8b12cd67f924289c02d127c4d17cdca5 --hash=sha256:28cdba5dcf2d03c89bb25dc3b2f5770ac4104470ed5bbe680a15494fa52a537d --hash=sha256:49a6e27abafacef104c7ca336f6790f91c69617a1d752ead4a017b706d632b55 --hash=sha256:7146c6fec3b2d7e8d11fa2504046b186d396520c0c2f7ef3aed40d8456dbebfc --hash=sha256:76d966194d51852f48c42a6e786a4520f1e1be5f93e2626423d673663422d559 --hash=sha256:7fc0f8788bc84433d2124e9a4893adbe40f93c7d213abb8ad7b909540cb0161f --hash=sha256:8698a59ad03539a2982125b7998efc1c107ba31d5d03437b6fcd72cb2c226922 --hash=sha256:8f0c4f7c42425177168938386a4fabd734ca3bbb5de2d1fd1176cfa6f980fc13 --hash=sha256:b6c4972f3ddd10a3de6c30311da2c25bc493d023796e16c5d4e0f8bd6d5770be --hash=sha256:c5c2a5fe06c7021183d9274df798f25516409c8d55898c324dae8a0b2de10144 --hash=sha256:dad93893e3fcb2410bfd31e854f20e1370b4fbafa07a737f1046f5fbd29ba0f4 --hash=sha256:e801cd0d35b4e3e99fc6a6f15eb193ce907bfa78127afa5825f1fad24c700c0e",
 				"multipledispatch==0.6.0 --hash=sha256:a55c512128fb3f7c2efd2533f2550accb93c35f1045242ef74645fc92a2c3cba",
 				"PyOpenGL==3.1.6 --hash=sha256:a7139bc3e15d656feae1f7e3ef68c799941ed43fadc78177a23db7e946c20738",
 				"pyrr==0.10.3 --hash=sha256:d8af23fb9bb29262405845e1c98f7339fbba5e49323b98528bd01160a75c65ac",
-				"pyyaml==6.0 --hash=sha256:432557aa2c09802be39460360ddffd48156e30721f5e8d917f01d31694216782 --hash=sha256:81957921f441d50af23654aa6c5e5eaf9b06aba7f0a19c18a538dc7ef291c5a1 --hash=sha256:055d937d65826939cb044fc8c9b08889e8c743fdc6a32b33e2390f66013e449b --hash=sha256:d4db7c7aef085872ef65a8fd7d6d09a14ae91f691dec3e87ee5ee0539d516f53 --hash=sha256:f84fbc98b019fef2ee9a1cb3ce93e3187a6df0b2538a651bfb890254ba9f90b5 --hash=sha256:b3d267842bf12586ba6c734f89d1f5b871df0273157918b0ccefa29deb05c21c --hash=sha256:40527857252b61eacd1d9af500c3337ba8deb8fc298940291486c465c8b46ec0 --hash=sha256:d4b0ba9512519522b118090257be113b9468d804b19d63c71dbcf4a48fa32358 --hash=sha256:daf496c58a8c52083df09b80c860005194014c3698698d1a57cbcfa182142a3a --hash=sha256:01b45c0191e6d66c470b6cf1b9531a771a83c1c4208272ead47a3ae4f2f603bf --hash=sha256:9df7ed3b3d2e0ecfe09e14741b857df43adb5a3ddadc919a2d94fbdf78fea53c --hash=sha256:e61ceaab6f49fb8bdfaa0f92c4b57bcfbea54c09277b1b4f7ac376bfb7a7c174",
-				"scikit_learn==1.2.1 --hash=sha256:bed9f75763bd392c094bf474c7ab75a01d68b15146ea7a20c0f9ff6fb3063dad --hash=sha256:da0e2d50a8435ea8dc5cd21f1fc1a45d329bae03dcca92087ebed859d22d184e --hash=sha256:5523e21ab2b4d52b2bd41bedd335dbe8f3c1b5f6dd7c9c001b2e17ec9818af8d --hash=sha256:70fa30d146b7e9d0c256e73e271b3e17f23123b7c4adcbde1a385031adf59090 --hash=sha256:c9285275a435d1f8f47bbe3500346ab9ead2499e0e090518404d318ea90d1c1c --hash=sha256:c722f3446ad8c4f1a93b2399fe1a188635b94709a3f25e6f4d61efbe75fe8eaa --hash=sha256:a9abf17d177df54e529154f26acfd42930e19117d045e8a9a8e893ca82dd94ec --hash=sha256:d00e46a2a7fce6e118ed0f4c6263785bf6c297a94ffd0cd7b32455043c508cc8 --hash=sha256:d8bcd303dd982494842a3f482f844d539484c6043b4eed896b43ea8e5f609a21 --hash=sha256:e0ee4d4d32c94e082344308528f7b3c9294b60ab19c84eb37a2d9c88bdffd9d1 --hash=sha256:479aedd0abedbda6b8b4529145fe4cd8622f69f726a72cef8f75548a93eeb1e1 --hash=sha256:5a8111f3c7a314017ebf90d6feab861c11d1ca14f3dbafb39abcc31aa4c54ba6",
+				"pyyaml==6.0 --hash=sha256:01b45c0191e6d66c470b6cf1b9531a771a83c1c4208272ead47a3ae4f2f603bf --hash=sha256:055d937d65826939cb044fc8c9b08889e8c743fdc6a32b33e2390f66013e449b --hash=sha256:40527857252b61eacd1d9af500c3337ba8deb8fc298940291486c465c8b46ec0 --hash=sha256:432557aa2c09802be39460360ddffd48156e30721f5e8d917f01d31694216782 --hash=sha256:77f396e6ef4c73fdc33a9157446466f1cff553d979bd00ecb64385760c6babdc --hash=sha256:81957921f441d50af23654aa6c5e5eaf9b06aba7f0a19c18a538dc7ef291c5a1 --hash=sha256:9df7ed3b3d2e0ecfe09e14741b857df43adb5a3ddadc919a2d94fbdf78fea53c --hash=sha256:afa17f5bc4d1b10afd4466fd3a44dc0e245382deca5b3c353d8b757f9e3ecb8d --hash=sha256:b3d267842bf12586ba6c734f89d1f5b871df0273157918b0ccefa29deb05c21c --hash=sha256:d4b0ba9512519522b118090257be113b9468d804b19d63c71dbcf4a48fa32358 --hash=sha256:d4db7c7aef085872ef65a8fd7d6d09a14ae91f691dec3e87ee5ee0539d516f53 --hash=sha256:d67d839ede4ed1b28a4e8909735fc992a923cdb84e618544973d7dfc71540803 --hash=sha256:daf496c58a8c52083df09b80c860005194014c3698698d1a57cbcfa182142a3a --hash=sha256:e61ceaab6f49fb8bdfaa0f92c4b57bcfbea54c09277b1b4f7ac376bfb7a7c174 --hash=sha256:f84fbc98b019fef2ee9a1cb3ce93e3187a6df0b2538a651bfb890254ba9f90b5",
+				"scikit_learn==1.2.1 --hash=sha256:479aedd0abedbda6b8b4529145fe4cd8622f69f726a72cef8f75548a93eeb1e1 --hash=sha256:5523e21ab2b4d52b2bd41bedd335dbe8f3c1b5f6dd7c9c001b2e17ec9818af8d --hash=sha256:5a8111f3c7a314017ebf90d6feab861c11d1ca14f3dbafb39abcc31aa4c54ba6 --hash=sha256:61bb9c654b5d2e6cdd4b1c7e6048fc66270c1682bda1b0f7d2726fdae09010f4 --hash=sha256:70fa30d146b7e9d0c256e73e271b3e17f23123b7c4adcbde1a385031adf59090 --hash=sha256:a9abf17d177df54e529154f26acfd42930e19117d045e8a9a8e893ca82dd94ec --hash=sha256:bed9f75763bd392c094bf474c7ab75a01d68b15146ea7a20c0f9ff6fb3063dad --hash=sha256:c722f3446ad8c4f1a93b2399fe1a188635b94709a3f25e6f4d61efbe75fe8eaa --hash=sha256:c9285275a435d1f8f47bbe3500346ab9ead2499e0e090518404d318ea90d1c1c --hash=sha256:cba0c7c6bf1493f8ce670bab69f9317874826ee838988de377ae355abd4d74cf --hash=sha256:d00e46a2a7fce6e118ed0f4c6263785bf6c297a94ffd0cd7b32455043c508cc8 --hash=sha256:d8bcd303dd982494842a3f482f844d539484c6043b4eed896b43ea8e5f609a21 --hash=sha256:da0e2d50a8435ea8dc5cd21f1fc1a45d329bae03dcca92087ebed859d22d184e --hash=sha256:dc838b5a4057c55ba81b82316ea8bf443af445f96eb21500b0e40618017e0923 --hash=sha256:e0ee4d4d32c94e082344308528f7b3c9294b60ab19c84eb37a2d9c88bdffd9d1",
 				"six==1.16.0 --hash=sha256:8abb2f1d86890a2dfb989f9a77cfcfd3e47c2a354b01111771326f8aa26e0254",
 				"threadpoolctl==3.1.0 --hash=sha256:8b99adda265feb6773280df41eece7b2e6561b772d21ffd52e372f999024907b"
 			]
