Geometry facade for grooms and cards + new rendering + use of curve selection

#virtualized

diff --git a/Engine/Plugins/Experimental/HairCardGenerator/Source/HairCardGeneratorDataflow/HairCardGeneratorDataflow.Build.cs b/Engine/Plugins/Experimental/HairCardGenerator/Source/HairCardGeneratorDataflow/HairCardGeneratorDataflow.Build.cs
--- a/Engine/Plugins/Experimental/HairCardGenerator/Source/HairCardGeneratorDataflow/HairCardGeneratorDataflow.Build.cs
+++ b/Engine/Plugins/Experimental/HairCardGenerator/Source/HairCardGeneratorDataflow/HairCardGeneratorDataflow.Build.cs
@@ -27,7 +27,9 @@ namespace UnrealBuildTool.Rules
 					"HairCardGeneratorEditor",
 					"HairCardGeneratorFramework",
 					"Engine",
-					"GeometryCore"
+					"GeometryCore",
+					"PlanarCut",
+					"GeometryFramework"
 				}
 			);
 		}
diff --git a/Engine/Plugins/Experimental/HairCardGenerator/Source/HairCardGeneratorDataflow/Private/BuildCardsSettingsNode.cpp b/Engine/Plugins/Experimental/HairCardGenerator/Source/HairCardGeneratorDataflow/Private/BuildCardsSettingsNode.cpp
--- a/Engine/Plugins/Experimental/HairCardGenerator/Source/HairCardGeneratorDataflow/Private/BuildCardsSettingsNode.cpp
+++ b/Engine/Plugins/Experimental/HairCardGenerator/Source/HairCardGeneratorDataflow/Private/BuildCardsSettingsNode.cpp
@@ -1,10 +1,10 @@
 // Copyright Epic Games, Inc. All Rights Reserved.
 
 #include "BuildCardsSettingsNode.h"
-
-#include "GroomCollectionFacades.h"
 #include "HairCardGeneratorEditorModule.h"
 #include "Dataflow/DataflowDebugDrawInterface.h"
+#include "GeometryCollection/Facades/CollectionConstraintOverrideFacade.h"
+#include "GeometryCollection/Facades/CollectionCurveFacade.h"
 
 #include UE_INLINE_GENERATED_CPP_BY_NAME(BuildCardsSettingsNode)
 
@@ -12,23 +12,23 @@ namespace UE::CardGen::Private
 {
 	static void BuildStrandsPositions(const FManagedArrayCollection& GroomCollection, TArray<TArray<FVector>>& StrandsPositions)
 	{
-		const UE::Groom::FGroomStrandsFacade StrandsFacade(GroomCollection);
-		if (StrandsFacade.IsValid())
+		const GeometryCollection::Facades::FCollectionCurveGeometryFacade CurvesFacade(GroomCollection);
+		if (CurvesFacade.IsValid())
 		{
-			const int32 NumStrands = StrandsFacade.GetNumCurves();
-			StrandsPositions.SetNum(NumStrands);
+			const int32 NumCurves = CurvesFacade.GetNumCurves();
+			StrandsPositions.SetNum(NumCurves);
 
 			int32 PointOffset = 0;
-			for(int32 StrandIndex = 0; StrandIndex < NumStrands; ++StrandIndex)
+			for(int32 CurveIndex = 0; CurveIndex < NumCurves; ++CurveIndex)
 			{
-				const int32 NumPoints = StrandsFacade.GetCurvePointOffsets()[StrandIndex]-PointOffset;
-				StrandsPositions[StrandIndex].SetNum(NumPoints);
+				const int32 NumPoints = CurvesFacade.GetCurvePointOffsets()[CurveIndex]-PointOffset;
+				StrandsPositions[CurveIndex].SetNum(NumPoints);
 				
 				for(int32 PointIndex = 0; PointIndex < NumPoints; ++PointIndex)
 				{
-					StrandsPositions[StrandIndex][PointIndex] = FVector(StrandsFacade.GetPointRestPositions()[PointIndex + PointOffset]);
+					StrandsPositions[CurveIndex][PointIndex] = FVector(CurvesFacade.GetPointRestPositions()[PointIndex + PointOffset]);
 				}
-				PointOffset = StrandsFacade.GetCurvePointOffsets()[StrandIndex];
+				PointOffset = CurvesFacade.GetCurvePointOffsets()[CurveIndex];
 			}
 		}
 	}
@@ -36,46 +36,45 @@ namespace UE::CardGen::Private
 	static FString DisplayCardsGroups(const UGroomAsset* GroomAsset)
 	{
 		TArray<FName> GroupNames;
-		FHairCardGeneratorUtils::BuildCardsGroups(GroomAsset, GroupNames);
-		
-		bool bFirstGroup = true;
+		TArray<int32> GroupCount;
+		FHairCardGeneratorUtils::BuildCardsGroups(GroomAsset, GroupNames, GroupCount);
+
 		FString CardsGroups;
-		CardsGroups += "[";
-		for(const FName& GroupName : GroupNames)
+		if(GroupNames.Num() == GroupCount.Num())
 		{
-			if(!bFirstGroup)
+			for(int32 GroupIndex = 0; GroupIndex < GroupNames.Num(); ++GroupIndex)
 			{
-				CardsGroups +=  ",";
+				CardsGroups += FString("Cards Groups Name = ") + GroupNames[GroupIndex].ToString() + FString(" | Num Strands = ") +
+					FString::FromInt(GroupCount[GroupIndex]) + FString("\n");
 			}
-			CardsGroups += GroupName.ToString();
-			bFirstGroup = false;
 		}
-		CardsGroups += "]";
 		return CardsGroups;
 	}
 }
 
+FBuildCardsSettingsNode::FBuildCardsSettingsNode(const UE::Dataflow::FNodeParameters& InParam, FGuid InGuid)
+	: FDataflowNode(InParam, InGuid)
+{
+	RegisterInputConnection(&Collection);
+	RegisterInputConnection(&GroomAsset);
+	RegisterOutputConnection(&Collection, &Collection);
+	RegisterOutputConnection(&CardsSettings);
+}
+
 void FBuildCardsSettingsNode::Evaluate(UE::Dataflow::FContext& Context, const FDataflowOutput* Out) const
 {
 	if (Out->IsA<FManagedArrayCollection>(&Collection))
 	{
-		FManagedArrayCollection GroomCollection = GetValue<FManagedArrayCollection>(Context, &Collection);
-		SetValue(Context, MoveTemp(GroomCollection), &Collection);
+		FManagedArrayCollection GeometryCollection = GetValue<FManagedArrayCollection>(Context, &Collection);
+		SetValue(Context, MoveTemp(GeometryCollection), &Collection);
 	}
 	else if (Out->IsA<TArray<FGroomCardsSettings>>(&CardsSettings))
 	{
-		UGroomAsset* LocalAsset = const_cast<UGroomAsset*>(GroomAsset.Get());
-		if (!LocalAsset)
-		{
-			if (const UE::Dataflow::FEngineContext* EngineContext = Context.AsType<UE::Dataflow::FEngineContext>())
-			{
-				LocalAsset = Cast<UGroomAsset>(EngineContext->Owner);
-			}
-		}
+		UGroomAsset* LocalAsset = GetValue<TObjectPtr<UGroomAsset>>(Context, &GroomAsset);
 		TArray<FGroomCardsSettings> OutputSettings;
 		if(LocalAsset && !LocalAsset->GetHairGroupsCards().IsEmpty())
 		{
-			FManagedArrayCollection GroomCollection = GetValue<FManagedArrayCollection>(Context, &Collection);
+			FManagedArrayCollection GeometryCollection = GetValue<FManagedArrayCollection>(Context, &Collection);
 			for(FHairGroupsCardsSourceDescription& HairCardsDescription : LocalAsset->GetHairGroupsCards())
 			{
 				if(HairCardsDescription.LODIndex != INDEX_NONE)
@@ -100,11 +99,28 @@ void FBuildCardsSettingsNode::Evaluate(UE::Dataflow::FContext& Context, const FD
 							}
 						}
 						OutputSettings.Last().GenerationSettings->BuildFilterGroupSettings(FilterCardGroups, FilterGroupNames);
+	
+						for(const FGroomFilterSettings& LocalSettings : FilterSettings)
+						{
+							for(TObjectPtr<UHairCardGeneratorGroupSettings> FilterGroups : OutputSettings.Last().GenerationSettings->GetFilterGroupSettings())
+							{
+								if (FilterGroups.Get())
+								{
+									if ((LocalSettings.FilterName != NAME_None && FilterGroups->GetFilterName() == LocalSettings.FilterName) || LocalSettings.FilterName == NAME_None)
+									{
+										FilterGroups->TargetNumberOfCards = LocalSettings.NumClumps * 3;
+										FilterGroups->MaxNumberOfFlyaways = LocalSettings.NumFlyaways;
+										FilterGroups->TargetTriangleCount = LocalSettings.NumTriangles;
+										FilterGroups->NumberOfTexturesInAtlas = LocalSettings.NumTextures;
+									}
+								}
+							}
+						}
 					}
 
-					FHairCardGeneratorUtils::LoadGroomStrands(LocalAsset, [&GroomCollection](TArray<TArray<FVector>>& StrandsPositions)
+					FHairCardGeneratorUtils::LoadGroomStrands(LocalAsset, [&GeometryCollection](TArray<TArray<FVector>>& CurvesPositions)
 					{
-						UE::CardGen::Private::BuildStrandsPositions(GroomCollection, StrandsPositions);
+						UE::CardGen::Private::BuildStrandsPositions(GeometryCollection, CurvesPositions);
 					});
 					OutputSettings.Last().GenerationSettings->UpdateStrandFilterAssignment();
 				}
@@ -135,7 +151,7 @@ void FBuildCardsSettingsNode::DebugDraw(UE::Dataflow::FContext& Context, IDatafl
 		if(LocalAsset)
 		{
 			FString CardsInfos = FString("Groom Asset Name : ") + LocalAsset->GetName() + FString("\n");
-			CardsInfos += FString("Cards Groups Names : ") + UE::CardGen::Private::DisplayCardsGroups(LocalAsset) + FString("\n");
+			CardsInfos += UE::CardGen::Private::DisplayCardsGroups(LocalAsset);
 			CardsInfos +=  FString("Num Cards Generations : ") + FString::FromInt(LocalAsset->GetHairGroupsCards().Num());
 			DataflowRenderingInterface.DrawOverlayText(CardsInfos);
 		}
diff --git a/Engine/Plugins/Experimental/HairCardGenerator/Source/HairCardGeneratorDataflow/Private/CardsAssetTerminalNode.cpp b/Engine/Plugins/Experimental/HairCardGenerator/Source/HairCardGeneratorDataflow/Private/CardsAssetTerminalNode.cpp
--- a/Engine/Plugins/Experimental/HairCardGenerator/Source/HairCardGeneratorDataflow/Private/CardsAssetTerminalNode.cpp
+++ b/Engine/Plugins/Experimental/HairCardGenerator/Source/HairCardGeneratorDataflow/Private/CardsAssetTerminalNode.cpp
@@ -3,7 +3,6 @@
 #include "CardsAssetTerminalNode.h"
 
 #include "BuildCardsSettingsNode.h"
-#include "GroomCollectionFacades.h"
 #include "GroomEdit.h"
 #include "HairCardGenControllerBase.h"
 #include "HairCardGeneratorEditorModule.h"
@@ -12,8 +11,13 @@
 
 #include UE_INLINE_GENERATED_CPP_BY_NAME(CardsAssetTerminalNode)
 
-namespace UE::Groom::Private
+
+FCardsAssetTerminalNode::FCardsAssetTerminalNode(const UE::Dataflow::FNodeParameters& InParam, FGuid InGuid)
+	: FDataflowTerminalNode(InParam, InGuid)
 {
+	RegisterInputConnection(&Collection);
+	RegisterInputConnection(&CardsSettings);
+	RegisterOutputConnection(&Collection, &Collection);
 }
 
 void FCardsAssetTerminalNode::SetAssetValue(TObjectPtr<UObject> Asset, UE::Dataflow::FContext& Context) const
diff --git a/Engine/Plugins/Experimental/HairCardGenerator/Source/HairCardGeneratorDataflow/Private/GenerateCardsClumpsNode.cpp b/Engine/Plugins/Experimental/HairCardGenerator/Source/HairCardGeneratorDataflow/Private/GenerateCardsClumpsNode.cpp
--- a/Engine/Plugins/Experimental/HairCardGenerator/Source/HairCardGeneratorDataflow/Private/GenerateCardsClumpsNode.cpp
+++ b/Engine/Plugins/Experimental/HairCardGenerator/Source/HairCardGeneratorDataflow/Private/GenerateCardsClumpsNode.cpp
@@ -1,18 +1,28 @@
 // Copyright Epic Games, Inc. All Rights Reserved.
 
 #include "GenerateCardsClumpsNode.h"
-#include "GroomCollectionFacades.h"
 #include "HairCardGeneratorEditorModule.h"
 #include "Dataflow/DataflowDebugDrawInterface.h"
 #include "Dataflow/DataflowDebugDrawObject.h"
 #include "Dataflow/DataflowRenderingViewMode.h"
 #include "Generators/MinimalBoxMeshGenerator.h"
+#include "GeometryCollection/Facades/CollectionCurveFacade.h"
 
 #include UE_INLINE_GENERATED_CPP_BY_NAME(GenerateCardsClumpsNode)
 
 // Clumps Attributes
-const FName FGenerateCardsClumpsNode::CurveClumpIndicesAttribute("CurveClumpIndices_LOD");
-const FName FGenerateCardsClumpsNode::ObjectNumClumpsAttribute("ObjectNumClumps_LOD");
+const FName FGenerateCardsClumpsNode::CurveClumpIndicesAttribute("ClumpIndices_LOD");
+const FName FGenerateCardsClumpsNode::ObjectNumClumpsAttribute("NumClumps_LOD");
+const FName FGenerateCardsClumpsNode::CurveFilterIndicesAttribute("FilterIndices_LOD");
+
+FGenerateCardsClumpsNode::FGenerateCardsClumpsNode(const UE::Dataflow::FNodeParameters& InParam, FGuid InGuid)
+	: FDataflowNode(InParam, InGuid)
+{
+	RegisterInputConnection(&Collection);
+	RegisterInputConnection(&CardsSettings);
+	RegisterOutputConnection(&Collection, &Collection);
+	RegisterOutputConnection(&CardsSettings, &CardsSettings);
+}
 
 void FGenerateCardsClumpsNode::Evaluate(UE::Dataflow::FContext& Context, const FDataflowOutput* Out) const
 {
@@ -47,7 +57,7 @@ void FGenerateCardsClumpsNode::Evaluate(UE::Dataflow::FContext& Context, const F
 		}
 		else if(Out->IsA<FManagedArrayCollection>(&Collection))
 		{
-			FManagedArrayCollection GroomCollection = GetValue<FManagedArrayCollection>(Context, &Collection);
+			FManagedArrayCollection GeometryCollection = GetValue<FManagedArrayCollection>(Context, &Collection);
 			
 			for(FGroomCardsSettings& LODSettings : OutputSettings)
 			{
@@ -56,12 +66,14 @@ void FGenerateCardsClumpsNode::Evaluate(UE::Dataflow::FContext& Context, const F
 					if(FHairCardGeneratorUtils::LoadGenerationSettings(GenerationSettings))
 					{
 						TArray<int32> StrandsClumps;
+						TArray<int32> StrandsFilter;
 						int32 NumClumps = 0;
 
-						StrandsClumps.Init(INDEX_NONE, GroomCollection.NumElements(UE::Groom::FGroomStrandsFacade::CurvesGroup));
+						StrandsClumps.Init(INDEX_NONE, GeometryCollection.NumElements(GeometryCollection::Facades::FCollectionCurveGeometryFacade::CurvesGroup));
+						StrandsFilter.Init(INDEX_NONE, GeometryCollection.NumElements(GeometryCollection::Facades::FCollectionCurveGeometryFacade::CurvesGroup));
 						
 						bool bHasClumps = FHairCardGeneratorUtils::RunCardsGeneration(GenerationSettings, LODSettings.PipelineFlags,
-				[&StrandsClumps, &NumClumps](const TObjectPtr<const UHairCardGeneratorPluginSettings>& GenerationSettings, const int32 FilterIndex, const uint8 GenFlags)
+				[&StrandsClumps, &NumClumps, &StrandsFilter](const TObjectPtr<const UHairCardGeneratorPluginSettings>& GenerationSettings, const int32 FilterIndex, const uint8 GenFlags)
 						{
 							TArray<int32> FilterClumps; int32 ClumpCount = 0;
 							if(FHairCardGeneratorUtils::GenerateCardsClumps(GenerationSettings, FilterIndex, GenFlags, FilterClumps, ClumpCount))
@@ -71,6 +83,7 @@ void FGenerateCardsClumpsNode::Evaluate(UE::Dataflow::FContext& Context, const F
 									if (StrandsClumps.IsValidIndex(CurveIndex) && (FilterClumps[CurveIndex] != INDEX_NONE))
 									{
 										StrandsClumps[CurveIndex] = FilterClumps[CurveIndex] + NumClumps;
+										StrandsFilter[CurveIndex] = FilterIndex;
 									}
 								}
 								NumClumps += ClumpCount;
@@ -86,13 +99,18 @@ void FGenerateCardsClumpsNode::Evaluate(UE::Dataflow::FContext& Context, const F
 			
 							FString NumClumpsLOD = ObjectNumClumpsAttribute.ToString();
 							NumClumpsLOD.AppendInt(GenerationSettings->GetLODIndex());
+
+ 							FString FilterIndicesLOD = CurveFilterIndicesAttribute.ToString();
+ 							FilterIndicesLOD.AppendInt(GenerationSettings->GetLODIndex());
 					
-							TManagedArray<int32>& CurveClumpIndices = GroomCollection.AddAttribute<int32>(FName(ClumpIndicesLOD), UE::Groom::FGroomStrandsFacade::CurvesGroup);
-							TManagedArray<int32>& ObjectNumClumps = GroomCollection.AddAttribute<int32>(FName(NumClumpsLOD), UE::Groom::FGroomStrandsFacade::ObjectsGroup);
-							
+							TManagedArray<int32>& CurveClumpIndices = GeometryCollection.AddAttribute<int32>(FName(ClumpIndicesLOD), GeometryCollection::Facades::FCollectionCurveGeometryFacade::CurvesGroup);
+							TManagedArray<int32>& ObjectNumClumps = GeometryCollection.AddAttribute<int32>(FName(NumClumpsLOD), FGeometryCollection::GeometryGroup);
+ 							TManagedArray<int32>& CurveFilterIndices = GeometryCollection.AddAttribute<int32>(FName(FilterIndicesLOD), GeometryCollection::Facades::FCollectionCurveGeometryFacade::CurvesGroup);
+ 							
 							for(int32 CurveIndex = 0, NumCurves = CurveClumpIndices.Num(); CurveIndex < NumCurves; ++CurveIndex)
 							{
 								CurveClumpIndices[CurveIndex] = StrandsClumps[CurveIndex];
+								CurveFilterIndices[CurveIndex] = StrandsFilter[CurveIndex];
 							}
 							for(int32 ObjectIndex = 0, NumObjects = ObjectNumClumps.Num(); ObjectIndex < NumObjects; ++ObjectIndex)
 							{
@@ -102,7 +120,7 @@ void FGenerateCardsClumpsNode::Evaluate(UE::Dataflow::FContext& Context, const F
 					}
 				}
 			}
-			SetValue(Context, MoveTemp(GroomCollection), &Collection);
+			SetValue(Context, MoveTemp(GeometryCollection), &Collection);
 		}
 	}
 }
diff --git a/Engine/Plugins/Experimental/HairCardGenerator/Source/HairCardGeneratorDataflow/Private/GenerateCardsGeometryNode.cpp b/Engine/Plugins/Experimental/HairCardGenerator/Source/HairCardGeneratorDataflow/Private/GenerateCardsGeometryNode.cpp
--- a/Engine/Plugins/Experimental/HairCardGenerator/Source/HairCardGeneratorDataflow/Private/GenerateCardsGeometryNode.cpp
+++ b/Engine/Plugins/Experimental/HairCardGenerator/Source/HairCardGeneratorDataflow/Private/GenerateCardsGeometryNode.cpp
@@ -2,7 +2,6 @@
 
 #include "GenerateCardsGeometryNode.h"
 #include "GenerateCardsClumpsNode.h"
-#include "GroomCollectionFacades.h"
 #include "HairCardGeneratorEditorModule.h"
 
 #include UE_INLINE_GENERATED_CPP_BY_NAME(GenerateCardsGeometryNode)
@@ -15,6 +14,15 @@ const FName FGenerateCardsGeometryNode::VertexCardIndicesAttribute("VertexCardIn
 const FName FGenerateCardsGeometryNode::CardsVerticesGroup("CardsVertices_LOD");
 const FName FGenerateCardsGeometryNode::CardsFacesGroup("CardsFaces_LOD");
 
+FGenerateCardsGeometryNode::FGenerateCardsGeometryNode(const UE::Dataflow::FNodeParameters& InParam, FGuid InGuid)
+	: FDataflowNode(InParam, InGuid)
+{
+	RegisterInputConnection(&Collection);
+	RegisterInputConnection(&CardsSettings);
+	RegisterOutputConnection(&Collection, &Collection);
+	RegisterOutputConnection(&CardsSettings, &CardsSettings);
+}
+
 void FGenerateCardsGeometryNode::Evaluate(UE::Dataflow::FContext& Context, const FDataflowOutput* Out) const
 {
 	if (Out->IsA<FManagedArrayCollection>(&Collection) || Out->IsA<TArray<FGroomCardsSettings>>(&CardsSettings))
@@ -46,7 +54,7 @@ void FGenerateCardsGeometryNode::Evaluate(UE::Dataflow::FContext& Context, const
 		}
 		else if(Out->IsA<FManagedArrayCollection>(&Collection))
 		{
-			FManagedArrayCollection GroomCollection = GetValue<FManagedArrayCollection>(Context, &Collection);
+			FManagedArrayCollection GeometryCollection = GetValue<FManagedArrayCollection>(Context, &Collection);
 			
 			for(FGroomCardsSettings& LODSettings : OutputSettings)
 			{
@@ -99,15 +107,15 @@ void FGenerateCardsGeometryNode::Evaluate(UE::Dataflow::FContext& Context, const
 							FString CardsFacesLODGroup = CardsFacesGroup.ToString();
 							CardsFacesLODGroup.AppendInt(GenerationSettings->GetLODIndex());
 							
-							TManagedArray<FVector3f>& VertexClumpPositions = GroomCollection.AddAttribute<FVector3f>(VertexClumpPositionsAttribute, FName(CardsVerticesLODGroup));
-							TManagedArray<FIntVector3>& FaceVertexIndices = GroomCollection.AddAttribute<FIntVector3>(FaceVertexIndicesAttribute, FName(CardsFacesLODGroup));
-							TManagedArray<int32>& VertexCardIndices = GroomCollection.AddAttribute<int32>(VertexCardIndicesAttribute, FName(CardsVerticesLODGroup));
+							TManagedArray<FVector3f>& VertexClumpPositions = GeometryCollection.AddAttribute<FVector3f>(VertexClumpPositionsAttribute, FName(CardsVerticesLODGroup));
+							TManagedArray<FIntVector3>& FaceVertexIndices = GeometryCollection.AddAttribute<FIntVector3>(FaceVertexIndicesAttribute, FName(CardsFacesLODGroup));
+							TManagedArray<int32>& VertexCardIndices = GeometryCollection.AddAttribute<int32>(VertexCardIndicesAttribute, FName(CardsVerticesLODGroup));
 							
-							GroomCollection.EmptyGroup(FName(CardsVerticesLODGroup));
-							GroomCollection.AddElements(ClumpsVertices.Num(), FName(CardsVerticesLODGroup));
+							GeometryCollection.EmptyGroup(FName(CardsVerticesLODGroup));
+							GeometryCollection.AddElements(ClumpsVertices.Num(), FName(CardsVerticesLODGroup));
 							
-							GroomCollection.EmptyGroup(FName(CardsFacesLODGroup));
-                            GroomCollection.AddElements(ClumpsFaces.Num(), FName(CardsFacesLODGroup));
+							GeometryCollection.EmptyGroup(FName(CardsFacesLODGroup));
+                            GeometryCollection.AddElements(ClumpsFaces.Num(), FName(CardsFacesLODGroup));
 							
 							for(int32 VertexIndex = 0, NumVertices = ClumpsVertices.Num(); VertexIndex < NumVertices; ++VertexIndex)
 							{
@@ -122,7 +130,7 @@ void FGenerateCardsGeometryNode::Evaluate(UE::Dataflow::FContext& Context, const
 					}
 				}
 			}
-			SetValue(Context, MoveTemp(GroomCollection), &Collection);
+			SetValue(Context, MoveTemp(GeometryCollection), &Collection);
 		}
 	}
 }
diff --git a/Engine/Plugins/Experimental/HairCardGenerator/Source/HairCardGeneratorDataflow/Private/GenerateCardsTexturesNode.cpp b/Engine/Plugins/Experimental/HairCardGenerator/Source/HairCardGeneratorDataflow/Private/GenerateCardsTexturesNode.cpp
--- a/Engine/Plugins/Experimental/HairCardGenerator/Source/HairCardGeneratorDataflow/Private/GenerateCardsTexturesNode.cpp
+++ b/Engine/Plugins/Experimental/HairCardGenerator/Source/HairCardGeneratorDataflow/Private/GenerateCardsTexturesNode.cpp
@@ -3,7 +3,6 @@
 #include "GenerateCardsTexturesNode.h"
 
 #include "GenerateCardsGeometryNode.h"
-#include "GroomCollectionFacades.h"
 #include "HairCardGeneratorEditorModule.h"
 
 #include UE_INLINE_GENERATED_CPP_BY_NAME(GenerateCardsTexturesNode)
@@ -13,6 +12,15 @@ const FName FGenerateCardsTexturesNode::ObjectTextureIndicesAttribute("ObjectTex
 const FName FGenerateCardsTexturesNode::CardsObjectsGroup("CardsObjects_LOD");
 const FName FGenerateCardsTexturesNode::VertexTextureUVsAttribute("VertexTextureUVs");
 
+FGenerateCardsTexturesNode::FGenerateCardsTexturesNode(const UE::Dataflow::FNodeParameters& InParam, FGuid InGuid)
+	: FDataflowNode(InParam, InGuid)
+{
+	RegisterInputConnection(&Collection);
+	RegisterInputConnection(&CardsSettings);
+	RegisterOutputConnection(&Collection, &Collection);
+	RegisterOutputConnection(&CardsSettings, &CardsSettings);
+}
+
 void FGenerateCardsTexturesNode::Evaluate(UE::Dataflow::FContext& Context, const FDataflowOutput* Out) const
 {
 	if (Out->IsA<FManagedArrayCollection>(&Collection) || Out->IsA<TArray<FGroomCardsSettings>>(&CardsSettings))
diff --git a/Engine/Plugins/Experimental/HairCardGenerator/Source/HairCardGeneratorDataflow/Private/HairCardDataflowRendering.cpp b/Engine/Plugins/Experimental/HairCardGenerator/Source/HairCardGeneratorDataflow/Private/HairCardDataflowRendering.cpp
--- a/Engine/Plugins/Experimental/HairCardGenerator/Source/HairCardGeneratorDataflow/Private/HairCardDataflowRendering.cpp
+++ b/Engine/Plugins/Experimental/HairCardGenerator/Source/HairCardGeneratorDataflow/Private/HairCardDataflowRendering.cpp
@@ -1,11 +1,24 @@
 // Copyright Epic Games, Inc. All Rights Reserved.
 
 #include "HairCardDataflowRendering.h"
+#include "Components/DynamicMeshComponent.h"
 #include "GenerateCardsGeometryNode.h"
 #include "GenerateCardsClumpsNode.h"
 #include "GenerateCardsTexturesNode.h"
-#include "GroomCollectionFacades.h"
+#include "PlanarCut.h"
+#include "Components/SphereComponent.h"
+#include "Components/TextRenderComponent.h"
+#include "Dataflow/DataflowRenderingViewMode.h"
+#include "DataflowRendering/DataflowRenderableComponents.h"
+#include "DataflowRendering/DataflowRenderableType.h"
+#include "DataflowRendering/DataflowRenderableTypeInstance.h"
+#include "DataflowRendering/DataflowRenderableTypeRegistry.h"
+#include "DataflowRendering/DataflowRenderableTypeUtils.h"
+#include "GeometryCollection/Facades/CollectionCurveFacade.h"
 #include "GeometryCollection/Facades/CollectionRenderingFacade.h"
+#include "Materials/Material.h"
+
+#define LOCTEXT_NAMESPACE "CardsRendering"
 
 namespace UE::CardGen::Private
 {
@@ -22,422 +35,690 @@ namespace UE::CardGen::Private
 		 GGroomDataflowCardsAlpha,
 		 TEXT("Cards Alpha for the rendering")
 	 );
-	
-	Dataflow::FRenderKey FCardsGeometryRenderingCallbacks::RenderKey = { "GeometryRender", FName("FCardsCollection") };
-	
-	void FCardsGeometryRenderingCallbacks::Render(GeometryCollection::Facades::FRenderingFacade& RenderCollection, const Dataflow::FGraphRenderingState& State)
+
+	static void BuildGeometryCollection(const TArray<FVector3f>& InVertexPosisions, const TArray<FVector3f>& InVertexNormals,
+		const TArray<FIntVector3>& InFaceIndices, const TArray<FLinearColor>& InVertexColors, const FName& GeometryName, Geometry::FDynamicMesh3& DynamicMesh)
 	{
-		if (State.GetRenderOutputs().Num())
+		FGeometryCollection GeometryCollection;
+		
+		// Transform group
+		GeometryCollection.EmptyGroup(FGeometryCollection::TransformGroup);
+		GeometryCollection.AddElements(1, FGeometryCollection::TransformGroup);
+
+		// Geometry group
+		GeometryCollection.EmptyGroup(FGeometryCollection::GeometryGroup);
+		GeometryCollection.AddElements(1, FGeometryCollection::GeometryGroup);
+
+		// Vertices group
+		GeometryCollection.EmptyGroup(FGeometryCollection::VerticesGroup);
+		GeometryCollection.AddElements(InVertexPosisions.Num(), FGeometryCollection::VerticesGroup);
+
+		// Faces group
+		GeometryCollection.EmptyGroup(FGeometryCollection::FacesGroup);
+		GeometryCollection.AddElements(InFaceIndices.Num(), FGeometryCollection::FacesGroup);
+
+		// Material group
+		GeometryCollection.EmptyGroup(FGeometryCollection::MaterialGroup);
+
+		// Vertices Attributes
+		TManagedArray<FVector3f>& VertexPositions = GeometryCollection.ModifyAttribute<FVector3f>(FGeometryCollection::VertexPositionAttribute, FGeometryCollection::VerticesGroup);
+		TManagedArray<FVector3f>& VertexNormals = GeometryCollection.ModifyAttribute<FVector3f>(FGeometryCollection::VertexNormalAttribute, FGeometryCollection::VerticesGroup);
+		TManagedArray<FVector3f>& VertexTangentU = GeometryCollection.ModifyAttribute<FVector3f>(FGeometryCollection::VertexTangentUAttribute, FGeometryCollection::VerticesGroup);
+		TManagedArray<FVector3f>& VertexTangentV = GeometryCollection.ModifyAttribute<FVector3f>(FGeometryCollection::VertexTangentVAttribute, FGeometryCollection::VerticesGroup);
+		TManagedArray<FLinearColor>& VertexColors = GeometryCollection.ModifyAttribute<FLinearColor>(FGeometryCollection::ColorAttribute, FGeometryCollection::VerticesGroup);
+		TManagedArray<int32>& VertexTransforms = GeometryCollection.ModifyAttribute<int32>(FGeometryCollection::VertexBoneMapAttribute, FGeometryCollection::VerticesGroup);
+		
+		// Faces Attributes
+		TManagedArray<FIntVector>& FaceIndices = GeometryCollection.ModifyAttribute<FIntVector>(FGeometryCollection::FaceIndicesAttribute, FGeometryCollection::FacesGroup);
+		TManagedArray<bool>& FaceVisibles = GeometryCollection.ModifyAttribute<bool>(FGeometryCollection::FaceVisibleAttribute, FGeometryCollection::FacesGroup);
+
+		// Geometry Group
+		TManagedArray<int32>& TransformIndex = GeometryCollection.ModifyAttribute<int32>(FGeometryCollection::TransformIndexAttribute, FGeometryCollection::GeometryGroup);
+		TManagedArray<FBox>& BoundingBox =  GeometryCollection.ModifyAttribute<FBox>(FGeometryCollection::BoundingBoxAttribute, FGeometryCollection::GeometryGroup);
+		TManagedArray<int32>& VertexStart = GeometryCollection.ModifyAttribute<int32>(FGeometryCollection::VertexStartAttribute, FGeometryCollection::GeometryGroup);
+		TManagedArray<int32>& VertexCount = GeometryCollection.ModifyAttribute<int32>(FGeometryCollection::VertexCountAttribute, FGeometryCollection::GeometryGroup);
+		TManagedArray<int32>& FaceStart = GeometryCollection.ModifyAttribute<int32>(FGeometryCollection::FaceStartAttribute, FGeometryCollection::GeometryGroup);
+		TManagedArray<int32>& FaceCount = GeometryCollection.ModifyAttribute<int32>(FGeometryCollection::FaceCountAttribute, FGeometryCollection::GeometryGroup);
+
+		// Transform group
+		TManagedArray<int32>& BoneGeometry = GeometryCollection.ModifyAttribute<int32>(FTransformCollection::GeometryIndexAttribute, FTransformCollection::TransformGroup);
+		TManagedArray<FTransform3f>& BoneTransforms = GeometryCollection.ModifyAttribute<FTransform3f>(FTransformCollection::TransformAttribute, FTransformCollection::TransformGroup);
+		TManagedArray<FString>& BoneNames = GeometryCollection.ModifyAttribute<FString>(FTransformCollection::BoneNameAttribute, FTransformCollection::TransformGroup);
+		TManagedArray<int32>& BoneParent = GeometryCollection.ModifyAttribute<int32>(FTransformCollection::ParentAttribute, FTransformCollection::TransformGroup);
+
+		VertexTangentU.Fill(FVector3f::Zero());
+		VertexTangentV.Fill(FVector3f::Zero());
+		VertexTransforms.Fill(0);
+		FaceVisibles.Fill(true);
+		
+		VertexNormals = InVertexNormals;
+		VertexPositions = InVertexPosisions;
+		VertexColors = InVertexColors;
+		FaceIndices = InFaceIndices;
+		
+		BoundingBox[0] = FBox(ForceInitToZero);
+		TransformIndex[0] = 0;
+		BoneGeometry[0] = 0;
+		BoneTransforms[0] = FTransform3f::Identity;
+		BoneParent[0] = INDEX_NONE;
+		BoneNames[0] = GeometryName.ToString();
+		VertexStart[0] = 0;
+		FaceStart[0] = 0;
+		VertexCount[0] = InVertexPosisions.Num();
+		FaceCount[0] = InFaceIndices.Num();
+
+		for(int32 VertexIndex = 0; VertexIndex < InVertexPosisions.Num(); ++VertexIndex)
 		{
-			const FManagedArrayCollection Default;
-			checkf(State.GetRenderOutputs().Num() == 1, TEXT("Expected FGraphRenderingState object to have one render output"));
-			const FName PrimaryOutput = State.GetRenderOutputs()[0];
-			const FManagedArrayCollection& Collection = State.GetValue<FManagedArrayCollection>(PrimaryOutput, Default);
-
-			FString CardsVerticesLODGroup = FGenerateCardsGeometryNode::CardsVerticesGroup.ToString();
-			CardsVerticesLODGroup.AppendInt(GGroomDataflowCardsLod);
-							
-			FString CardsFacesLODGroup = FGenerateCardsGeometryNode::CardsFacesGroup.ToString();
-			CardsFacesLODGroup.AppendInt(GGroomDataflowCardsLod);
+			BoundingBox[0] += FVector(VertexPositions[VertexIndex]);
+		}
+		
+		TArrayView<const int32> TransformIndicesView = {0};
+		TArrayView<const FTransform3f> BoneTransformsView = MakeArrayView(BoneTransforms.GetConstArray());
 
-			if (Collection.HasAttribute(FGenerateCardsGeometryNode::VertexClumpPositionsAttribute, FName(CardsVerticesLODGroup)) &&
-				Collection.HasAttribute(FGenerateCardsGeometryNode::FaceVertexIndicesAttribute, FName(CardsFacesLODGroup)) && 
-				Collection.HasAttribute(FGenerateCardsGeometryNode::VertexCardIndicesAttribute, FName(CardsVerticesLODGroup)))
+		FTransform UnusedTransform;
+		::ConvertGeometryCollectionToDynamicMesh(DynamicMesh, UnusedTransform, false, GeometryCollection, true, BoneTransformsView, true, TransformIndicesView);
+	}
+	
+	static void AddMeshComponents(TArray<Geometry::FDynamicMesh3>& DynamicMeshes, Dataflow::FRenderableComponents& OutComponents, const TArray<TArray<int32>>& FilterTransforms)
+	{
+		if(!FilterTransforms.IsEmpty())
+		{
+			int32 FilterIndex = 0;
+			for(const TArray<int32>& TransformsIndices : FilterTransforms)
 			{
-				if (State.GetViewMode().GetName() == Dataflow::FDataflowConstruction3DViewMode::Name)
+				if (USphereComponent* FilterComponent = OutComponents.AddNewComponent<USphereComponent>())
 				{
-					const TManagedArray<FVector3f>& VertexGlobalPositions = Collection.GetAttribute<FVector3f>(
-						FGenerateCardsGeometryNode::VertexClumpPositionsAttribute, FName(CardsVerticesLODGroup));
-						
-					const TManagedArray<FIntVector3>& FaceGlobalVertices = Collection.GetAttribute<FIntVector3>(
-						FGenerateCardsGeometryNode::FaceVertexIndicesAttribute, FName(CardsFacesLODGroup));
-
-					const TManagedArray<int32>& VertexCardIndices = Collection.GetAttribute<int32>(
-						FGenerateCardsGeometryNode::VertexCardIndicesAttribute, FName(CardsVerticesLODGroup));
-
-					int32 NumCards = 0;
-					for(const int32& CardIndex : VertexCardIndices.GetConstArray())
-					{
-						NumCards = FMath::Max(NumCards, CardIndex);
-					}
-					++NumCards;
-
-					TArray<TArray<int32>> CardVertexIndices;
-					TArray<TArray<int32>> CardFaceIndices;
-					TArray<int32> VertexLocalIndices;
-
-					CardVertexIndices.SetNum(NumCards);
-					CardFaceIndices.SetNum(NumCards);
-					VertexLocalIndices.SetNum(VertexCardIndices.Num());
-					
-					for (int32 VertexIndex = 0, NumVertices = VertexCardIndices.Num(); VertexIndex < NumVertices; ++VertexIndex)
+					int32 LocalIndex = 0;
+					for(const int32& TransformIndex : TransformsIndices)
 					{
-						VertexLocalIndices[VertexIndex] = CardVertexIndices[VertexCardIndices[VertexIndex]].Num();
-						CardVertexIndices[VertexCardIndices[VertexIndex]].Add(VertexIndex);
-					}
-
-					for (int32 FaceIndex = 0, NumFaces = FaceGlobalVertices.Num(); FaceIndex < NumFaces; ++FaceIndex)
-					{
-						CardFaceIndices[VertexCardIndices[FaceGlobalVertices[FaceIndex][0]]].Add(FaceIndex);
-					}
-
-					for(int32 CardIndex = 0; CardIndex < NumCards; ++CardIndex)
-					{	
-						TArray<FVector3f> VertexLocalPositions;
-						TArray<FVector3f> VertexLocalNormals;
-						TArray<FIntVector> FaceLocalVertices;
-						TArray<int32> VertexNumFaces;
-
-						VertexLocalPositions.Init(FVector3f::Zero(), CardVertexIndices[CardIndex].Num());
-						VertexLocalNormals.Init(FVector3f::Zero(), CardVertexIndices[CardIndex].Num());
-						VertexNumFaces.Init(0, CardVertexIndices[CardIndex].Num());
-						FaceLocalVertices.Init(FIntVector::ZeroValue, CardFaceIndices.Num());
-
-						for (int32 FaceIndex = 0, NumFaces = CardFaceIndices[CardIndex].Num(); FaceIndex < NumFaces; ++FaceIndex)
+						if (DynamicMeshes[TransformIndex].VertexCount() > 0)
 						{
-							const FIntVector3 GlobalVertices = FaceGlobalVertices[CardFaceIndices[CardIndex][FaceIndex]];
-							FaceLocalVertices[FaceIndex] = FIntVector(VertexLocalIndices[GlobalVertices[0]], VertexLocalIndices[GlobalVertices[1]], VertexLocalIndices[GlobalVertices[2]]);
-
-							const FVector3f FaceNormal = (VertexGlobalPositions[GlobalVertices[2]] - VertexGlobalPositions[GlobalVertices[0]]).Cross(
-								VertexGlobalPositions[GlobalVertices[1]] - VertexGlobalPositions[GlobalVertices[0]]).GetSafeNormal();
-
-							VertexLocalNormals[FaceLocalVertices[FaceIndex][0]] += FaceNormal;
-							VertexLocalNormals[FaceLocalVertices[FaceIndex][1]] += FaceNormal;
-							VertexLocalNormals[FaceLocalVertices[FaceIndex][2]] += FaceNormal;
-
-							VertexNumFaces[FaceLocalVertices[FaceIndex][0]]++;
-							VertexNumFaces[FaceLocalVertices[FaceIndex][1]]++;
-							VertexNumFaces[FaceLocalVertices[FaceIndex][2]]++;
-						}
-
-						for (int32 VertexIndex = 0, NumVertices = CardVertexIndices[CardIndex].Num(); VertexIndex < NumVertices; ++VertexIndex)
-						{
-							if (VertexNumFaces[VertexIndex] > 0)
+							if (UDynamicMeshComponent* Component = OutComponents.AddNewComponent<UDynamicMeshComponent>(NAME_None, FilterComponent))
 							{
-								VertexLocalNormals[VertexIndex] /= VertexNumFaces[VertexIndex];
+								Component->SetupAttachment(FilterComponent);
+								Component->SetMesh(MoveTemp(DynamicMeshes[TransformIndex]));
+								Component->SetOverrideRenderMaterial(Dataflow::Rendering::GetVertexMaterial());
+								Component->SetShadowsEnabled(false);
 							}
-							VertexLocalPositions[VertexIndex] = VertexGlobalPositions[CardVertexIndices[CardIndex][VertexIndex]];
 						}
-
-						const FLinearColor GroupColor = FLinearColor::IntToDistinctColor(CardIndex, 0.75f, 1.0f, 90.0f);
-
-						TArray<FLinearColor> VertexLocalColors;
-						VertexLocalColors.Init(GroupColor, VertexLocalPositions.Num());
-
-						ComputeVertexColors(Collection, CardIndex, VertexLocalColors);
-
-						const FString GeometryName = TEXT("Groom_Card_") + FString::FromInt(CardIndex);
-						const int32 GeometryIndex = RenderCollection.StartGeometryGroup(GeometryName);
-						RenderCollection.AddSurface(MoveTemp(VertexLocalPositions), MoveTemp(FaceLocalVertices), MoveTemp(VertexLocalNormals), MoveTemp(VertexLocalColors));
-						RenderCollection.EndGeometryGroup(GeometryIndex);
+						++LocalIndex;
 					}
 				}
-				else
+				++FilterIndex;
+			}
+		}
+		else
+		{
+			for (int32 TransformIndex = 0; TransformIndex < DynamicMeshes.Num(); ++TransformIndex)
+			{
+				if (DynamicMeshes[TransformIndex].VertexCount() > 0)
 				{
-					checkf(false, TEXT("Invalid View Mode for groom cards geometry rendering"));
+					if (UDynamicMeshComponent* Component = OutComponents.AddNewComponent<UDynamicMeshComponent>())
+					{
+						Component->SetMesh(MoveTemp(DynamicMeshes[TransformIndex]));
+						Component->SetOverrideRenderMaterial(Dataflow::Rendering::GetVertexMaterial());
+						Component->SetShadowsEnabled(false);
+					}
 				}
 			}
 		}
 	}
 
-	Dataflow::FRenderKey FCardsTextureRenderingCallbacks::RenderKey = { "TextureRender", FName("FCardsCollection") };
-
-	void FCardsTextureRenderingCallbacks::Render(GeometryCollection::Facades::FRenderingFacade& RenderCollection, const Dataflow::FGraphRenderingState& State)
+	static void CountClumpsAndFilters(const TArray<int32>& CurveClumpIndices, const TArray<int32>& CurveFilterIndices,
+		TArray<TArray<int32>>& FilterClumpIndices, TArray<TArray<int32>>& ClumpCurveIndices, int32& NumClumps, int32& NumFilters)
 	{
-		if (State.GetViewMode().GetName() == Dataflow::FDataflowConstruction3DViewMode::Name)
+		for(int32 CurveIndex = 0, NumCurves = CurveClumpIndices.Num(); CurveIndex < NumCurves; ++CurveIndex)
 		{
-			FCardsGeometryRenderingCallbacks::Render(RenderCollection, State);
+			NumClumps = FMath::Max(NumClumps, CurveClumpIndices[CurveIndex]);
+			NumFilters = FMath::Max(NumFilters, CurveFilterIndices[CurveIndex]);
 		}
-		else if (State.GetViewMode().GetName() == Dataflow::FDataflowConstructionUVViewMode::Name)
+		++NumClumps;
+		++NumFilters;
+		
+		ClumpCurveIndices.SetNum(NumClumps);
+		FilterClumpIndices.SetNum(NumFilters);
+
+		// Fill the clump curve indices and filter clump indices
+		for(int32 CurveIndex = 0, NumCurves = CurveClumpIndices.Num(); CurveIndex < NumCurves; ++CurveIndex)
 		{
-			if (State.GetRenderOutputs().Num())
+			const int32 ClumpIndex = CurveClumpIndices[CurveIndex];
+			if ((ClumpIndex >= 0) && ClumpCurveIndices.IsValidIndex(ClumpIndex))
 			{
-				const FManagedArrayCollection Default;
-				checkf(State.GetRenderOutputs().Num() == 1, TEXT("Expected FGraphRenderingState object to have one render output"));
-				const FName PrimaryOutput = State.GetRenderOutputs()[0];
-				const FManagedArrayCollection& Collection = State.GetValue<FManagedArrayCollection>(PrimaryOutput, Default);
-
-				FString CardsVerticesLODGroup = FGenerateCardsGeometryNode::CardsVerticesGroup.ToString();
-				CardsVerticesLODGroup.AppendInt(GGroomDataflowCardsLod);
-							
-				FString CardsFacesLODGroup = FGenerateCardsGeometryNode::CardsFacesGroup.ToString();
-				CardsFacesLODGroup.AppendInt(GGroomDataflowCardsLod);
-
-				if (Collection.HasAttribute(FGenerateCardsTexturesNode::VertexTextureUVsAttribute, FName(CardsVerticesLODGroup)) &&
-					Collection.HasAttribute(FGenerateCardsGeometryNode::FaceVertexIndicesAttribute, FName(CardsFacesLODGroup)))
+				ClumpCurveIndices[ClumpIndex].Add(CurveIndex);
+			}
+		}
+		for(int32 ClumpIndex = 0; ClumpIndex < NumClumps; ++ClumpIndex)
+		{
+			if(!ClumpCurveIndices[ClumpIndex].IsEmpty())
+			{
+				const int32 CurveIndex = ClumpCurveIndices[ClumpIndex].Last();
+				if(CurveFilterIndices.IsValidIndex(CurveIndex))
 				{
-					const TManagedArray<FVector2f>& VertexUVsArray = Collection.GetAttribute<FVector2f>(
-							FGenerateCardsTexturesNode::VertexTextureUVsAttribute, FName(CardsVerticesLODGroup));
-						
-					const TManagedArray<FIntVector3>& FaceIndicesArray = Collection.GetAttribute<FIntVector3>(
-						FGenerateCardsGeometryNode::FaceVertexIndicesAttribute, FName(CardsFacesLODGroup));
-					
-					TArray<FIntVector> FaceIndices = FaceIndicesArray.GetConstArray();
-					
-					TArray<FVector3f> VertexUVs;
-					VertexUVs.Init(FVector3f(0.0, 0.0, 1.0), VertexUVsArray.Num());
-
-					for (int32 VertexIndex = 0; VertexIndex < VertexUVsArray.Num(); ++VertexIndex)
+					const int32 FilterIndex = CurveFilterIndices[CurveIndex];
+					if(FilterClumpIndices.IsValidIndex(FilterIndex))
 					{
-						VertexUVs[VertexIndex] = FVector3f(VertexUVsArray[VertexIndex][0], VertexUVsArray[VertexIndex][1], 0.0f);
+						FilterClumpIndices[FilterIndex].Add(ClumpIndex);
 					}
-
-					TArray<FVector3f> VertexNormals;
-					VertexNormals.Init(FVector3f(0.0, 0.0, 1.0), VertexUVs.Num());
-
-					TArray<FLinearColor> VertexColors;
-					VertexColors.Init(FLinearColor(0,0,0,0), VertexUVs.Num());
-
-					const int32 GeometryIndex = RenderCollection.StartGeometryGroup(State.GetGuid().ToString());
-					RenderCollection.AddSurface(MoveTemp(VertexUVs), MoveTemp(FaceIndices), MoveTemp(VertexNormals), MoveTemp(VertexColors));
-					RenderCollection.EndGeometryGroup(GeometryIndex);
-				}	
+				}
 			}
 		}
 	}
 
-	void FCardsTextureRenderingCallbacks::ComputeVertexColors(const FManagedArrayCollection& Collection, const int32 CardIndex, TArray<FLinearColor>& VertexColors) const
+	static void BuildCardsGeometry(const FManagedArrayCollection& Collection, TArray<TArray<FVector3f>>& VertexPositions, TArray<TArray<FVector3f>>& VertexNormals, TArray<TArray<FIntVector3>>& FaceVertices)
 	{
-		FString CardsObjectsLODGroup = FGenerateCardsTexturesNode::CardsObjectsGroup.ToString();
-		CardsObjectsLODGroup.AppendInt(GGroomDataflowCardsLod);
+		FString CardsVerticesLODGroup = FGenerateCardsGeometryNode::CardsVerticesGroup.ToString();
+		CardsVerticesLODGroup.AppendInt(GGroomDataflowCardsLod);
+					
+		FString CardsFacesLODGroup = FGenerateCardsGeometryNode::CardsFacesGroup.ToString();
+		CardsFacesLODGroup.AppendInt(GGroomDataflowCardsLod);
 
-		if (Collection.HasAttribute(FGenerateCardsTexturesNode::ObjectTextureIndicesAttribute, FName(CardsObjectsLODGroup)))
+		if (Collection.HasAttribute(FGenerateCardsGeometryNode::VertexClumpPositionsAttribute, FName(CardsVerticesLODGroup)) &&
+			Collection.HasAttribute(FGenerateCardsGeometryNode::FaceVertexIndicesAttribute, FName(CardsFacesLODGroup)) && 
+			Collection.HasAttribute(FGenerateCardsGeometryNode::VertexCardIndicesAttribute, FName(CardsVerticesLODGroup)))
 		{
-			const TManagedArray<int32>& CardsTextureArray = Collection.GetAttribute<int32>(
-				FGenerateCardsTexturesNode::ObjectTextureIndicesAttribute, FName(CardsObjectsLODGroup));
+			const TManagedArray<FVector3f>& VertexGlobalPositions = Collection.GetAttribute<FVector3f>(
+				FGenerateCardsGeometryNode::VertexClumpPositionsAttribute, FName(CardsVerticesLODGroup));
+				
+			const TManagedArray<FIntVector3>& FaceGlobalVertices = Collection.GetAttribute<FIntVector3>(
+				FGenerateCardsGeometryNode::FaceVertexIndicesAttribute, FName(CardsFacesLODGroup));
+
+			const TManagedArray<int32>& VertexCardIndices = Collection.GetAttribute<int32>(
+				FGenerateCardsGeometryNode::VertexCardIndicesAttribute, FName(CardsVerticesLODGroup));
+
+			int32 NumCards = 0;
+			for(const int32& CardIndex : VertexCardIndices.GetConstArray())
+			{
+				NumCards = FMath::Max(NumCards, CardIndex);
+			}
+			++NumCards;
 
-			const int32 TextureIndex = (CardIndex < CardsTextureArray.Num()) ? CardsTextureArray[CardIndex] : INDEX_NONE;
-			FLinearColor CardColor = (TextureIndex != INDEX_NONE) ? FLinearColor::IntToDistinctColor(TextureIndex, 0.75f, 1.0f, 90.0f) : FLinearColor::Black;
-			CardColor *= (TextureIndex == CardIndex) ? 1.0 : GGroomDataflowCardsAlpha;
+			TArray<TArray<int32>> CardVertexIndices;
+			TArray<TArray<int32>> CardFaceIndices;
+			TArray<int32> VertexLocalIndices;
+
+			CardVertexIndices.SetNum(NumCards);
+			CardFaceIndices.SetNum(NumCards);
+			VertexLocalIndices.SetNum(VertexCardIndices.Num());
+			
+			for (int32 VertexIndex = 0, NumVertices = VertexCardIndices.Num(); VertexIndex < NumVertices; ++VertexIndex)
+			{
+				VertexLocalIndices[VertexIndex] = CardVertexIndices[VertexCardIndices[VertexIndex]].Num();
+				CardVertexIndices[VertexCardIndices[VertexIndex]].Add(VertexIndex);
+			}
+
+			for (int32 FaceIndex = 0, NumFaces = FaceGlobalVertices.Num(); FaceIndex < NumFaces; ++FaceIndex)
+			{
+				CardFaceIndices[VertexCardIndices[FaceGlobalVertices[FaceIndex][0]]].Add(FaceIndex);
+			}
+			VertexPositions.SetNum(NumCards);
+			VertexNormals.SetNum(NumCards);
+			FaceVertices.SetNum(NumCards);
+			for(int32 CardIndex = 0; CardIndex < NumCards; ++CardIndex)
+			{
+				TArray<int32> VertexNumFaces;
 
-			const int32 NumVertices = VertexColors.Num();
-			VertexColors.Init(CardColor, NumVertices);
+				VertexPositions[CardIndex].Init(FVector3f::Zero(), CardVertexIndices[CardIndex].Num());
+				VertexNormals[CardIndex].Init(FVector3f::Zero(), CardVertexIndices[CardIndex].Num());
+				VertexNumFaces.Init(0, CardVertexIndices[CardIndex].Num());
+				FaceVertices[CardIndex].Init(FIntVector::ZeroValue, CardFaceIndices.Num());
+
+				for (int32 FaceIndex = 0, NumFaces = CardFaceIndices[CardIndex].Num(); FaceIndex < NumFaces; ++FaceIndex)
+				{
+					const FIntVector3 GlobalVertices = FaceGlobalVertices[CardFaceIndices[CardIndex][FaceIndex]];
+					FaceVertices[CardIndex][FaceIndex] = FIntVector(VertexLocalIndices[GlobalVertices[0]], VertexLocalIndices[GlobalVertices[1]], VertexLocalIndices[GlobalVertices[2]]);
+
+					const FVector3f FaceNormal = (VertexGlobalPositions[GlobalVertices[2]] - VertexGlobalPositions[GlobalVertices[0]]).Cross(
+						VertexGlobalPositions[GlobalVertices[1]] - VertexGlobalPositions[GlobalVertices[0]]).GetSafeNormal();
+
+					VertexNormals[CardIndex][FaceVertices[CardIndex][FaceIndex][0]] += FaceNormal;
+					VertexNormals[CardIndex][FaceVertices[CardIndex][FaceIndex][1]] += FaceNormal;
+					VertexNormals[CardIndex][FaceVertices[CardIndex][FaceIndex][2]] += FaceNormal;
+
+					VertexNumFaces[FaceVertices[CardIndex][FaceIndex][0]]++;
+					VertexNumFaces[FaceVertices[CardIndex][FaceIndex][1]]++;
+					VertexNumFaces[FaceVertices[CardIndex][FaceIndex][2]]++;
+				}
+
+				for (int32 VertexIndex = 0, NumVertices = CardVertexIndices[CardIndex].Num(); VertexIndex < NumVertices; ++VertexIndex)
+				{
+					if (VertexNumFaces[VertexIndex] > 0)
+					{
+						VertexNormals[CardIndex][VertexIndex] /= VertexNumFaces[VertexIndex];
+					}
+					VertexPositions[CardIndex][VertexIndex] = VertexGlobalPositions[CardVertexIndices[CardIndex][VertexIndex]];
+				}
+			}
 		}
 	}
 
-	Dataflow::FRenderKey FCardsClumpsRenderingCallbacks::RenderKey = { "ClumpsRender", FName("FCardsCollection") };
-	
-	void FCardsClumpsRenderingCallbacks::Render(GeometryCollection::Facades::FRenderingFacade& RenderCollection, const Dataflow::FGraphRenderingState& State)
+	class FGeometryCollectionCardsClumpsRenderableType : public Dataflow::IRenderableType
 	{
-		if (State.GetRenderOutputs().Num())
+		UE_DATAFLOW_IRENDERABLE_PRIMARY_TYPE(FManagedArrayCollection, Collection);
+		UE_DATAFLOW_IRENDERABLE_RENDER_GROUP(CardsClumps);
+		UE_DATAFLOW_IRENDERABLE_VIEW_MODE(UE::Dataflow::FDataflowConstruction3DViewMode);
+
+		virtual bool CanRender(const Dataflow::FRenderableTypeInstance& Instance) const
 		{
-			const FManagedArrayCollection Default;
-			checkf(State.GetRenderOutputs().Num() == 1, TEXT("Expected FGraphRenderingState object to have one render output"));
-			const FName PrimaryOutput = State.GetRenderOutputs()[0];
-			const FManagedArrayCollection& Collection = State.GetValue<FManagedArrayCollection>(PrimaryOutput, Default);
+			const FManagedArrayCollection& Collection = GetCollection(Instance);
+			return Collection.HasGroup(FGeometryCollection::GeometryGroup);
+		}
+
+		virtual void GetPrimitiveComponents(const Dataflow::FRenderableTypeInstance& Instance, Dataflow::FRenderableComponents& OutComponents) const override
+		{
+			TArray<Geometry::FDynamicMesh3> DynamicMeshes;
+			
+			TArray<TArray<int32>> FilterClumpIndices;
+			TArray<TArray<int32>> ClumpCurveIndices;
+			
+			const FManagedArrayCollection& Collection = GetCollection(Instance);
 
 			FString ClumpIndicesLOD = FGenerateCardsClumpsNode::CurveClumpIndicesAttribute.ToString();
 			ClumpIndicesLOD.AppendInt(GGroomDataflowCardsLod);
 			
-			FString NumClumpsLOD = FGenerateCardsClumpsNode::ObjectNumClumpsAttribute.ToString();
-			NumClumpsLOD.AppendInt(GGroomDataflowCardsLod);
+			FString FilterIndicesLOD = FGenerateCardsClumpsNode::CurveFilterIndicesAttribute.ToString();
+			FilterIndicesLOD.AppendInt(GGroomDataflowCardsLod);
 
-			if (Collection.HasAttribute(FName(ClumpIndicesLOD), UE::Groom::FGroomStrandsFacade::CurvesGroup) &&
-				Collection.HasAttribute(FName(NumClumpsLOD), UE::Groom::FGroomStrandsFacade::ObjectsGroup))
+			if (Collection.HasAttribute(FName(ClumpIndicesLOD), GeometryCollection::Facades::FCollectionCurveGeometryFacade::CurvesGroup) &&
+				Collection.HasAttribute(FName(FilterIndicesLOD), GeometryCollection::Facades::FCollectionCurveGeometryFacade::CurvesGroup))
 			{
-				if (State.GetViewMode().GetName() == Dataflow::FDataflowConstruction3DViewMode::Name)
+				const GeometryCollection::Facades::FCollectionCurveGeometryFacade CurvesFacade(Collection);
+				if (CurvesFacade.IsValid())
 				{
-					const UE::Groom::FGroomStrandsFacade StrandsFacade(Collection);
-					if (StrandsFacade.IsValid())
-					{
-						const TManagedArray<int32>& CurveClumpIndices = Collection.GetAttribute<int32>(FName(ClumpIndicesLOD), UE::Groom::FGroomStrandsFacade::CurvesGroup);
+					const TManagedArray<int32>& CurveClumpIndices = Collection.GetAttribute<int32>(FName(ClumpIndicesLOD), GeometryCollection::Facades::FCollectionCurveGeometryFacade::CurvesGroup);
+					const TManagedArray<int32>& CurveFilterIndices = Collection.GetAttribute<int32>(FName(FilterIndicesLOD), GeometryCollection::Facades::FCollectionCurveGeometryFacade::CurvesGroup);
 
-						int32 NumClumps = 0;
-						for(int32 CurveIndex = 0, NumCurves = CurveClumpIndices.Num(); CurveIndex < NumCurves; ++CurveIndex)
-						{
-							NumClumps = FMath::Max(NumClumps, CurveClumpIndices[CurveIndex]);
-						}
-						++NumClumps;
-
-						TArray<TArray<int32>> ClumpCurveIndices;
-						TArray<TArray<TArray<FVector>>> ClumpMeshPoints;
-						TArray<TArray<FVector>> ClumpMeshOrigins;
-						TArray<TArray<FQuat>> ClumpMeshRotations;
-						TArray<TArray<FVector>> ClumpMeshExtentsX;
-						TArray<TArray<FVector>> ClumpMeshExtentsY;
-
-						ClumpCurveIndices.SetNum(NumClumps);
-						ClumpMeshPoints.SetNum(NumClumps);
-						ClumpMeshOrigins.SetNum(NumClumps);
-						ClumpMeshRotations.SetNum(NumClumps);
-						ClumpMeshExtentsX.SetNum(NumClumps);
-						ClumpMeshExtentsY.SetNum(NumClumps);
-
-						// Fill the clump curve indices
-						for(int32 CurveIndex = 0, NumCurves = CurveClumpIndices.Num(); CurveIndex < NumCurves; ++CurveIndex)
-						{
-							if ((CurveClumpIndices[CurveIndex] >= 0) && ClumpCurveIndices.IsValidIndex(CurveClumpIndices[CurveIndex]))
-							{
-								ClumpCurveIndices[CurveClumpIndices[CurveIndex]].Add(CurveIndex);
-							}
-						}
-
-						static constexpr const int32 NumMeshes = 10;
-						static constexpr const int32 NumCorners = 8;
-						static constexpr const int32 NumTriangles = 12;
-						static constexpr const float MinExtent = 0.5f;
+					int32 NumClumps = 0, NumFilters = 0;
+					CountClumpsAndFilters(CurveClumpIndices.GetConstArray(), CurveFilterIndices.GetConstArray(),
+						FilterClumpIndices, ClumpCurveIndices, NumClumps, NumFilters);
+					
+					if (Instance.HasUptoDateCachedValue())
+					{
+						DynamicMeshes = Instance.GetCachedValue<TArray<Geometry::FDynamicMesh3>>();
+					}
+					else
+					{
+						FString NumClumpsLOD = FGenerateCardsClumpsNode::ObjectNumClumpsAttribute.ToString();
+						NumClumpsLOD.AppendInt(GGroomDataflowCardsLod);
 
-						for(int32 ClumpIndex = 0; ClumpIndex < NumClumps; ++ClumpIndex)
+						if (Collection.HasAttribute(FName(NumClumpsLOD), FGeometryCollection::GeometryGroup))
 						{
-							// Fill the clump points for eaqch meshes
-							ClumpMeshPoints[ClumpIndex].SetNum(NumMeshes+1);
-							for (const int32& CurveIndex : ClumpCurveIndices[ClumpIndex])
+							TArray<TArray<TArray<FVector>>> ClumpMeshPoints;
+							TArray<TArray<FVector>> ClumpMeshOrigins;
+							TArray<TArray<FQuat>> ClumpMeshRotations;
+							TArray<TArray<FVector>> ClumpMeshExtentsX;
+							TArray<TArray<FVector>> ClumpMeshExtentsY;
+
+							ClumpMeshPoints.SetNum(NumClumps);
+							ClumpMeshOrigins.SetNum(NumClumps);
+							ClumpMeshRotations.SetNum(NumClumps);
+							ClumpMeshExtentsX.SetNum(NumClumps);
+							ClumpMeshExtentsY.SetNum(NumClumps);
+
+							static constexpr const int32 NumMeshes = 10;
+							static constexpr const int32 NumCorners = 8;
+							static constexpr const int32 NumTriangles = 12;
+							static constexpr const float MinExtent = 0.5f;
+							
+							DynamicMeshes.SetNum(NumClumps);
+							for(int32 ClumpIndex = 0; ClumpIndex < NumClumps; ++ClumpIndex)
 							{
-								if(StrandsFacade.GetCurvePointOffsets().IsValidIndex(CurveIndex) && StrandsFacade.GetCurvePointOffsets().IsValidIndex(CurveIndex - 1))
+								// Fill the clump points for eaqch meshes
+								ClumpMeshPoints[ClumpIndex].SetNum(NumMeshes+1);
+								for (const int32& CurveIndex : ClumpCurveIndices[ClumpIndex])
 								{
-									const float PointStart = (CurveIndex > 0) ? StrandsFacade.GetCurvePointOffsets()[CurveIndex - 1] : 0;
-									const float PointEnd = StrandsFacade.GetCurvePointOffsets()[CurveIndex]-1;
-
-									if(StrandsFacade.GetPointRestPositions().IsValidIndex(PointStart) && StrandsFacade.GetPointRestPositions().IsValidIndex(PointEnd))
+									if(CurvesFacade.GetCurvePointOffsets().IsValidIndex(CurveIndex) && CurvesFacade.GetCurvePointOffsets().IsValidIndex(CurveIndex - 1))
 									{
-										ClumpMeshPoints[ClumpIndex][0].Add(FVector(StrandsFacade.GetPointRestPositions()[PointStart]));
-										ClumpMeshPoints[ClumpIndex][NumMeshes].Add(FVector(StrandsFacade.GetPointRestPositions()[PointEnd]));
+										const float PointStart = (CurveIndex > 0) ? CurvesFacade.GetCurvePointOffsets()[CurveIndex - 1] : 0;
+										const float PointEnd = CurvesFacade.GetCurvePointOffsets()[CurveIndex]-1;
 
-										for(int32 MeshIndex = 1; MeshIndex < NumMeshes; ++MeshIndex)
+										if(CurvesFacade.GetPointRestPositions().IsValidIndex(PointStart) && CurvesFacade.GetPointRestPositions().IsValidIndex(PointEnd))
 										{
-											const float PointParam = PointStart + (PointEnd - PointStart) * MeshIndex / NumMeshes;
-											const int32 ParamIndex = FMath::FloorToInt(PointParam);
-											const float ParamAlpha = PointParam - ParamIndex;
+											ClumpMeshPoints[ClumpIndex][0].Add(FVector(CurvesFacade.GetPointRestPositions()[PointStart]));
+											ClumpMeshPoints[ClumpIndex][NumMeshes].Add(FVector(CurvesFacade.GetPointRestPositions()[PointEnd]));
 
-											if(StrandsFacade.GetPointRestPositions().IsValidIndex(ParamIndex) && StrandsFacade.GetPointRestPositions().IsValidIndex(ParamIndex+1))
+											for(int32 MeshIndex = 1; MeshIndex < NumMeshes; ++MeshIndex)
 											{
-												ClumpMeshPoints[ClumpIndex][MeshIndex].Add(FVector(StrandsFacade.GetPointRestPositions()[ParamIndex] * (1.0 - ParamAlpha) + StrandsFacade.GetPointRestPositions()[ParamIndex+1] * ParamAlpha));
+												const float PointParam = PointStart + (PointEnd - PointStart) * MeshIndex / NumMeshes;
+												const int32 ParamIndex = FMath::FloorToInt(PointParam);
+												const float ParamAlpha = PointParam - ParamIndex;
+
+												if(CurvesFacade.GetPointRestPositions().IsValidIndex(ParamIndex) && CurvesFacade.GetPointRestPositions().IsValidIndex(ParamIndex+1))
+												{
+													ClumpMeshPoints[ClumpIndex][MeshIndex].Add(FVector(CurvesFacade.GetPointRestPositions()[ParamIndex] * (1.0 - ParamAlpha) + CurvesFacade.GetPointRestPositions()[ParamIndex+1] * ParamAlpha));
+												}
 											}
 										}
 									}
 								}
-							}
 
-							// Compute the clump mesh origins
-							ClumpMeshOrigins[ClumpIndex].Init(FVector::ZeroVector, NumMeshes + 1);
-							for (int32 MeshIndex = 0; MeshIndex < NumMeshes+1; ++MeshIndex)
-							{
-								if(!ClumpMeshPoints[ClumpIndex][MeshIndex].IsEmpty())
-								{ 
-									for (const FVector& MeshPoint : ClumpMeshPoints[ClumpIndex][MeshIndex])
-									{
-										ClumpMeshOrigins[ClumpIndex][MeshIndex] += MeshPoint;
+								// Compute the clump mesh origins
+								ClumpMeshOrigins[ClumpIndex].Init(FVector::ZeroVector, NumMeshes + 1);
+								for (int32 MeshIndex = 0; MeshIndex < NumMeshes+1; ++MeshIndex)
+								{
+									if(!ClumpMeshPoints[ClumpIndex][MeshIndex].IsEmpty())
+									{ 
+										for (const FVector& MeshPoint : ClumpMeshPoints[ClumpIndex][MeshIndex])
+										{
+											ClumpMeshOrigins[ClumpIndex][MeshIndex] += MeshPoint;
+										}
+										ClumpMeshOrigins[ClumpIndex][MeshIndex] /= ClumpMeshPoints[ClumpIndex][MeshIndex].Num();
 									}
-									ClumpMeshOrigins[ClumpIndex][MeshIndex] /= ClumpMeshPoints[ClumpIndex][MeshIndex].Num();
 								}
-							}
 
-							FVector NextTangent = FVector(0,0,1), PrevTangent = FVector(0, 0, 1);
-							FQuat MeshOrientation = FQuat::Identity;
+								FVector NextTangent = FVector(0,0,1), PrevTangent = FVector(0, 0, 1);
+								FQuat MeshOrientation = FQuat::Identity;
 
-							// Compute the clump mesh rotations
-							ClumpMeshRotations[ClumpIndex].Init(FQuat::Identity, NumMeshes);
-							for (int32 MeshIndex = 0; MeshIndex < NumMeshes; ++MeshIndex)
-							{
-								PrevTangent = NextTangent;
-								NextTangent = (ClumpMeshOrigins[ClumpIndex][MeshIndex+1]-ClumpMeshOrigins[ClumpIndex][MeshIndex]).GetSafeNormal();
+								// Compute the clump mesh rotations
+								ClumpMeshRotations[ClumpIndex].Init(FQuat::Identity, NumMeshes);
+								for (int32 MeshIndex = 0; MeshIndex < NumMeshes; ++MeshIndex)
+								{
+									PrevTangent = NextTangent;
+									NextTangent = (ClumpMeshOrigins[ClumpIndex][MeshIndex+1]-ClumpMeshOrigins[ClumpIndex][MeshIndex]).GetSafeNormal();
 
-								MeshOrientation = (FQuat::FindBetweenNormals(PrevTangent, NextTangent) * MeshOrientation).GetNormalized();
-								ClumpMeshRotations[ClumpIndex][MeshIndex] = MeshOrientation;
-							}
+									MeshOrientation = (FQuat::FindBetweenNormals(PrevTangent, NextTangent) * MeshOrientation).GetNormalized();
+									ClumpMeshRotations[ClumpIndex][MeshIndex] = MeshOrientation;
+								}
 
-							// Compute the clump mesh extents
-							ClumpMeshExtentsX[ClumpIndex].Init(FVector::ZeroVector, NumMeshes + 1);
-							ClumpMeshExtentsY[ClumpIndex].Init(FVector::ZeroVector, NumMeshes + 1);
+								// Compute the clump mesh extents
+								ClumpMeshExtentsX[ClumpIndex].Init(FVector::ZeroVector, NumMeshes + 1);
+								ClumpMeshExtentsY[ClumpIndex].Init(FVector::ZeroVector, NumMeshes + 1);
 
-							float MeshExtent = MinExtent;
-							for (int32 MeshIndex = 0; MeshIndex < NumMeshes+1; ++MeshIndex)
-							{
-								MeshOrientation = (MeshIndex == NumMeshes) ? ClumpMeshRotations[ClumpIndex][MeshIndex-1] : ClumpMeshRotations[ClumpIndex][MeshIndex];
+								float MeshExtent = MinExtent;
+								for (int32 MeshIndex = 0; MeshIndex < NumMeshes+1; ++MeshIndex)
+								{
+									MeshOrientation = (MeshIndex == NumMeshes) ? ClumpMeshRotations[ClumpIndex][MeshIndex-1] : ClumpMeshRotations[ClumpIndex][MeshIndex];
 
-								const FVector MeshNormalX = MeshOrientation.GetAxisX();
-								const FVector MeshNormalY = MeshOrientation.GetAxisY();
+									const FVector MeshNormalX = MeshOrientation.GetAxisX();
+									const FVector MeshNormalY = MeshOrientation.GetAxisY();
 
-								if(MeshIndex == 0)
-								{ 
-									FVector2d MeshExtents(MinExtent, MinExtent);
-									for (const FVector& MeshPoint : ClumpMeshPoints[ClumpIndex][MeshIndex])
-									{
-										const FVector LocalPoint = MeshPoint - ClumpMeshOrigins[ClumpIndex][MeshIndex];
+									if(MeshIndex == 0)
+									{ 
+										FVector2d MeshExtents(MinExtent, MinExtent);
+										for (const FVector& MeshPoint : ClumpMeshPoints[ClumpIndex][MeshIndex])
+										{
+											const FVector LocalPoint = MeshPoint - ClumpMeshOrigins[ClumpIndex][MeshIndex];
 
-										MeshExtents[0] = FMath::Max(MeshExtents[0], FMath::Abs(LocalPoint.Dot(MeshNormalX)));
-										MeshExtents[1] = FMath::Max(MeshExtents[1], FMath::Abs(LocalPoint.Dot(MeshNormalY)));
+											MeshExtents[0] = FMath::Max(MeshExtents[0], FMath::Abs(LocalPoint.Dot(MeshNormalX)));
+											MeshExtents[1] = FMath::Max(MeshExtents[1], FMath::Abs(LocalPoint.Dot(MeshNormalY)));
+										}
+										MeshExtent = FMath::Min(MeshExtents[0], MeshExtents[1]);
 									}
-									MeshExtent = FMath::Min(MeshExtents[0], MeshExtents[1]);
+									ClumpMeshExtentsX[ClumpIndex][MeshIndex] = MeshNormalX * MeshExtent;
+									ClumpMeshExtentsY[ClumpIndex][MeshIndex] = MeshNormalY * MeshExtent;
 								}
-								ClumpMeshExtentsX[ClumpIndex][MeshIndex] = MeshNormalX * MeshExtent;
-								ClumpMeshExtentsY[ClumpIndex][MeshIndex] = MeshNormalY * MeshExtent;
-							}
 
-							TArray<FVector3f> VertexLocalPositions;
-							TArray<FVector3f> VertexLocalNormals;
-							TArray<FIntVector> FaceLocalVertices;
+								TArray<FVector3f> VertexLocalPositions;
+								TArray<FVector3f> VertexLocalNormals;
+								TArray<FIntVector> FaceLocalVertices;
 
-							VertexLocalNormals.Init(FVector3f::ZeroVector, NumMeshes * NumCorners);
-							VertexLocalPositions.Reserve(NumMeshes * NumCorners);
-							FaceLocalVertices.Reserve(NumMeshes * NumTriangles);
+								VertexLocalNormals.Init(FVector3f::ZeroVector, NumMeshes * NumCorners);
+								VertexLocalPositions.Reserve(NumMeshes * NumCorners);
+								FaceLocalVertices.Reserve(NumMeshes * NumTriangles);
 
-							int32 VertexOffset = 0;
+								int32 VertexOffset = 0;
 
-							auto AddMeshFace = [&VertexOffset, &FaceLocalVertices, &VertexLocalNormals, &VertexLocalPositions](const int32 VertexIndexA, const int32 VertexIndexB, const int32 VertexIndexC)
-							{
-								FIntVector& FaceIndices = FaceLocalVertices.Add_GetRef(FIntVector(VertexOffset + VertexIndexA, VertexOffset + VertexIndexB, VertexOffset + VertexIndexC));
-								const FVector3f FaceNormal = ((VertexLocalPositions[FaceIndices[2]] - VertexLocalPositions[FaceIndices[0]]).Cross(VertexLocalPositions[FaceIndices[1]] - VertexLocalPositions[FaceIndices[0]])).GetSafeNormal();
-
-								VertexLocalNormals[FaceIndices[0]] += FaceNormal;
-								VertexLocalNormals[FaceIndices[1]] += FaceNormal;
-								VertexLocalNormals[FaceIndices[2]] += FaceNormal;
-							};
-							for (int32 MeshIndex = 0; MeshIndex < NumMeshes; ++MeshIndex)
-							{
-								// Add vertex positions
-								FVector3f MeshOrigin = FVector3f(ClumpMeshOrigins[ClumpIndex][MeshIndex]);
-								FVector3f MeshExtentX = FVector3f(ClumpMeshExtentsX[ClumpIndex][MeshIndex]);
-								FVector3f MeshExtentY = FVector3f(ClumpMeshExtentsY[ClumpIndex][MeshIndex]);
-
-								VertexLocalPositions.Add(MeshOrigin - MeshExtentX - MeshExtentY);
-								VertexLocalPositions.Add(MeshOrigin + MeshExtentX - MeshExtentY);
-								VertexLocalPositions.Add(MeshOrigin + MeshExtentX + MeshExtentY);
-								VertexLocalPositions.Add(MeshOrigin - MeshExtentX + MeshExtentY);
-
-								MeshOrigin = FVector3f(ClumpMeshOrigins[ClumpIndex][MeshIndex+1]);
-								MeshExtentX = FVector3f(ClumpMeshExtentsX[ClumpIndex][MeshIndex+1]);
-								MeshExtentY = FVector3f(ClumpMeshExtentsY[ClumpIndex][MeshIndex+1]);
-
-								VertexLocalPositions.Add(MeshOrigin - MeshExtentX - MeshExtentY);
-								VertexLocalPositions.Add(MeshOrigin + MeshExtentX - MeshExtentY);
-								VertexLocalPositions.Add(MeshOrigin + MeshExtentX + MeshExtentY);
-								VertexLocalPositions.Add(MeshOrigin - MeshExtentX + MeshExtentY);
-
-								// Add box faces and update normals
-								AddMeshFace(0, 1, 2);
-								AddMeshFace(2, 3, 0);
-								AddMeshFace(7, 6, 5);
-								AddMeshFace(5, 4, 7);
-								AddMeshFace(3, 2, 6);
-								AddMeshFace(6, 7, 3);
-								AddMeshFace(0, 3, 7);
-								AddMeshFace(7, 4, 0);
-								AddMeshFace(2, 1, 5);
-								AddMeshFace(5, 6, 2);
-								AddMeshFace(1, 0, 4);
-								AddMeshFace(4, 5, 1);
-
-								for(int32 VertexIndex = 0; VertexIndex < NumCorners; ++VertexIndex)
+								auto AddMeshFace = [&VertexOffset, &FaceLocalVertices, &VertexLocalNormals, &VertexLocalPositions](const int32 VertexIndexA, const int32 VertexIndexB, const int32 VertexIndexC)
+									{
+										FIntVector& FaceIndices = FaceLocalVertices.Add_GetRef(FIntVector(VertexOffset + VertexIndexA, VertexOffset + VertexIndexB, VertexOffset + VertexIndexC));
+										const FVector3f FaceNormal = ((VertexLocalPositions[FaceIndices[2]] - VertexLocalPositions[FaceIndices[0]]).Cross(VertexLocalPositions[FaceIndices[1]] - VertexLocalPositions[FaceIndices[0]])).GetSafeNormal();
+
+										VertexLocalNormals[FaceIndices[0]] += FaceNormal;
+										VertexLocalNormals[FaceIndices[1]] += FaceNormal;
+										VertexLocalNormals[FaceIndices[2]] += FaceNormal;
+									};
+								for (int32 MeshIndex = 0; MeshIndex < NumMeshes; ++MeshIndex)
 								{
-									VertexLocalNormals[VertexIndex + VertexOffset].Normalize();
+									// Add vertex positions
+									FVector3f MeshOrigin = FVector3f(ClumpMeshOrigins[ClumpIndex][MeshIndex]);
+									FVector3f MeshExtentX = FVector3f(ClumpMeshExtentsX[ClumpIndex][MeshIndex]);
+									FVector3f MeshExtentY = FVector3f(ClumpMeshExtentsY[ClumpIndex][MeshIndex]);
+
+									VertexLocalPositions.Add(MeshOrigin - MeshExtentX - MeshExtentY);
+									VertexLocalPositions.Add(MeshOrigin + MeshExtentX - MeshExtentY);
+									VertexLocalPositions.Add(MeshOrigin + MeshExtentX + MeshExtentY);
+									VertexLocalPositions.Add(MeshOrigin - MeshExtentX + MeshExtentY);
+
+									MeshOrigin = FVector3f(ClumpMeshOrigins[ClumpIndex][MeshIndex+1]);
+									MeshExtentX = FVector3f(ClumpMeshExtentsX[ClumpIndex][MeshIndex+1]);
+									MeshExtentY = FVector3f(ClumpMeshExtentsY[ClumpIndex][MeshIndex+1]);
+
+									VertexLocalPositions.Add(MeshOrigin - MeshExtentX - MeshExtentY);
+									VertexLocalPositions.Add(MeshOrigin + MeshExtentX - MeshExtentY);
+									VertexLocalPositions.Add(MeshOrigin + MeshExtentX + MeshExtentY);
+									VertexLocalPositions.Add(MeshOrigin - MeshExtentX + MeshExtentY);
+
+									// Add box faces and update normals
+									AddMeshFace(0, 1, 2);
+									AddMeshFace(2, 3, 0);
+									AddMeshFace(7, 6, 5);
+									AddMeshFace(5, 4, 7);
+									AddMeshFace(3, 2, 6);
+									AddMeshFace(6, 7, 3);
+									AddMeshFace(0, 3, 7);
+									AddMeshFace(7, 4, 0);
+									AddMeshFace(2, 1, 5);
+									AddMeshFace(5, 6, 2);
+									AddMeshFace(1, 0, 4);
+									AddMeshFace(4, 5, 1);
+
+									for(int32 VertexIndex = 0; VertexIndex < NumCorners; ++VertexIndex)
+									{
+										VertexLocalNormals[VertexIndex + VertexOffset].Normalize();
+									}
+									VertexOffset += NumCorners;
 								}
-								VertexOffset += NumCorners;
+
+								const FLinearColor GroupColor = FLinearColor::IntToDistinctColor(ClumpIndex, 0.75f, 1.0f, 90.0f);
+
+								TArray<FLinearColor> VertexLocalColors;
+								VertexLocalColors.Init(GroupColor, VertexLocalPositions.Num());
+								
+								const FString GeometryName = TEXT("Groom_Clump_") + FString::FromInt(ClumpIndex);
+
+								BuildGeometryCollection(VertexLocalPositions, VertexLocalNormals, FaceLocalVertices, VertexLocalColors, FName(GeometryName), DynamicMeshes[ClumpIndex]);
 							}
+							Instance.CacheValue(DynamicMeshes);
+						}
+					}
+				}
+			}
+			AddMeshComponents(DynamicMeshes, OutComponents, FilterClumpIndices);
+		}
+	};
 
-							const FLinearColor GroupColor = FLinearColor::IntToDistinctColor(ClumpIndex, 0.75f, 1.0f, 90.0f);
+	class FGeometryCollectionCardsGeometryRenderableType : public Dataflow::IRenderableType
+	{
+		UE_DATAFLOW_IRENDERABLE_PRIMARY_TYPE(FManagedArrayCollection, Collection);
+		UE_DATAFLOW_IRENDERABLE_RENDER_GROUP(CardsGeometry);
+		UE_DATAFLOW_IRENDERABLE_VIEW_MODE(UE::Dataflow::FDataflowConstruction3DViewMode);
+
+		virtual bool CanRender(const Dataflow::FRenderableTypeInstance& Instance) const
+		{
+			const FManagedArrayCollection& Collection = GetCollection(Instance);
+			return Collection.HasGroup(FGeometryCollection::GeometryGroup);
+		}
+
+		virtual void GetPrimitiveComponents(const Dataflow::FRenderableTypeInstance& Instance, Dataflow::FRenderableComponents& OutComponents) const override
+		{
+			TArray<Geometry::FDynamicMesh3> DynamicMeshes;
+			if (Instance.HasUptoDateCachedValue())
+			{
+				DynamicMeshes = Instance.GetCachedValue<TArray<Geometry::FDynamicMesh3>>();
+			}
+			else
+			{
+				const FManagedArrayCollection& Collection = GetCollection(Instance);
+
+				TArray<TArray<FVector3f>> VertexPositions, VertexNormals;
+				TArray<TArray<FIntVector3>> FaceVertices;
+				BuildCardsGeometry(Collection, VertexPositions, VertexNormals, FaceVertices);
+				
+				if(VertexPositions.Num() == VertexNormals.Num() && VertexNormals.Num() == FaceVertices.Num())
+				{
+					const int32 NumCards = VertexPositions.Num();
+
+					DynamicMeshes.SetNum(NumCards);
+					for(int32 CardIndex = 0; CardIndex < NumCards; ++CardIndex)
+					{	
+						const FLinearColor GroupColor = FLinearColor::IntToDistinctColor(CardIndex, 0.75f, 1.0f, 90.0f);
 
-							TArray<FLinearColor> VertexLocalColors;
-							VertexLocalColors.Init(GroupColor, VertexLocalPositions.Num());
+						TArray<FLinearColor> VertexColors;
+						VertexColors.Init(GroupColor, VertexPositions[CardIndex].Num());
 
-							const FString GeometryName = TEXT("Groom_Clump_") + FString::FromInt(ClumpIndex);
-							const int32 GeometryIndex = RenderCollection.StartGeometryGroup(GeometryName);
-							RenderCollection.AddSurface(MoveTemp(VertexLocalPositions), MoveTemp(FaceLocalVertices), MoveTemp(VertexLocalNormals), MoveTemp(VertexLocalColors));
-							RenderCollection.EndGeometryGroup(GeometryIndex);
+						const FString GeometryName = TEXT("Groom_Card_") + FString::FromInt(CardIndex);
+					
+						BuildGeometryCollection(VertexPositions[CardIndex], VertexNormals[CardIndex], FaceVertices[CardIndex], VertexColors, FName(GeometryName), DynamicMeshes[CardIndex]);
+					}
+					Instance.CacheValue(DynamicMeshes);
+				}
+			}
+			AddMeshComponents(DynamicMeshes, OutComponents, {});
+		}
+	};
+
+	class FGeometryCollectionCardsTextureRenderableType : public Dataflow::IRenderableType
+	{
+		UE_DATAFLOW_IRENDERABLE_PRIMARY_TYPE(FManagedArrayCollection, Collection);
+		UE_DATAFLOW_IRENDERABLE_RENDER_GROUP(CardsTexture);
+		UE_DATAFLOW_IRENDERABLE_VIEW_MODE(UE::Dataflow::FDataflowConstruction3DViewMode);
+
+		virtual bool CanRender(const Dataflow::FRenderableTypeInstance& Instance) const
+		{
+			const FManagedArrayCollection& Collection = GetCollection(Instance);
+			return Collection.HasGroup(FGeometryCollection::GeometryGroup);
+		}
+
+		virtual void GetPrimitiveComponents(const Dataflow::FRenderableTypeInstance& Instance, Dataflow::FRenderableComponents& OutComponents) const override
+		{
+			TArray<Geometry::FDynamicMesh3> DynamicMeshes;
+			TArray<TArray<int32>> FilterCardsIndices;
+			FilterCardsIndices.SetNum(2);
+			if (Instance.HasUptoDateCachedValue())
+			{
+				DynamicMeshes = Instance.GetCachedValue<TArray<Geometry::FDynamicMesh3>>();
+			}
+			else
+			{
+				const FManagedArrayCollection& Collection = GetCollection(Instance);
+				
+				TArray<TArray<FVector3f>> VertexPositions, VertexNormals;
+				TArray<TArray<FIntVector3>> FaceVertices;
+				BuildCardsGeometry(Collection, VertexPositions, VertexNormals, FaceVertices);
+
+				if(VertexPositions.Num() == VertexNormals.Num() && VertexNormals.Num() == FaceVertices.Num())
+				{
+					const int32 NumCards = VertexPositions.Num();
+
+					FString CardsObjectsLODGroup = FGenerateCardsTexturesNode::CardsObjectsGroup.ToString();
+					CardsObjectsLODGroup.AppendInt(GGroomDataflowCardsLod);
+
+					if (Collection.HasAttribute(FGenerateCardsTexturesNode::ObjectTextureIndicesAttribute, FName(CardsObjectsLODGroup)))
+					{
+						const TManagedArray<int32>& CardsTextureArray = Collection.GetAttribute<int32>(
+							FGenerateCardsTexturesNode::ObjectTextureIndicesAttribute, FName(CardsObjectsLODGroup));
+				
+						DynamicMeshes.SetNum(NumCards);
+						for(int32 CardIndex = 0; CardIndex < NumCards; ++CardIndex)
+						{
+							const int32 TextureIndex = (CardIndex < CardsTextureArray.Num()) ? CardsTextureArray[CardIndex] : INDEX_NONE;
+							FLinearColor CardColor = (TextureIndex != INDEX_NONE) ? FLinearColor::IntToDistinctColor(TextureIndex, 0.75f, 1.0f, 90.0f) : FLinearColor::Black;
+							CardColor *= (TextureIndex == CardIndex) ? 1.0 : GGroomDataflowCardsAlpha;
+
+							if(TextureIndex == CardIndex)
+							{
+								FilterCardsIndices[0].Add(CardIndex);
+							}
+							else
+							{
+								FilterCardsIndices[1].Add(CardIndex);
+							}
+
+							TArray<FLinearColor> VertexColors;
+							VertexColors.Init(CardColor, VertexPositions[CardIndex].Num());
+
+							const FString GeometryName = TEXT("Groom_Texture_") + FString::FromInt(CardIndex);
+							
+							BuildGeometryCollection(VertexPositions[CardIndex], VertexNormals[CardIndex], FaceVertices[CardIndex], VertexColors, FName(GeometryName), DynamicMeshes[CardIndex]);
 						}
+						Instance.CacheValue(DynamicMeshes);
+					}
+				}
+			}
+			AddMeshComponents(DynamicMeshes, OutComponents, FilterCardsIndices);
+		}
+	};
+
+	class FGeometryCollectionCardsUVsRenderableType : public Dataflow::IRenderableType
+	{
+		UE_DATAFLOW_IRENDERABLE_PRIMARY_TYPE(FManagedArrayCollection, Collection);
+		UE_DATAFLOW_IRENDERABLE_RENDER_GROUP(CardsTexture);
+		UE_DATAFLOW_IRENDERABLE_VIEW_MODE(UE::Dataflow::FDataflowConstructionUVViewMode);
+
+		virtual bool CanRender(const Dataflow::FRenderableTypeInstance& Instance) const
+		{
+			const FManagedArrayCollection& Collection = GetCollection(Instance);
+			return Collection.HasGroup(FGeometryCollection::GeometryGroup);
+		}
+
+		virtual void GetPrimitiveComponents(const Dataflow::FRenderableTypeInstance& Instance, Dataflow::FRenderableComponents& OutComponents) const override
+		{
+			TArray<Geometry::FDynamicMesh3> DynamicMeshes;
+			if (Instance.HasUptoDateCachedValue())
+			{
+				DynamicMeshes = Instance.GetCachedValue<TArray<Geometry::FDynamicMesh3>>();
+			}
+			else
+			{
+				const FManagedArrayCollection& Collection = GetCollection(Instance);
+				
+				FString CardsVerticesLODGroup = FGenerateCardsGeometryNode::CardsVerticesGroup.ToString();
+				CardsVerticesLODGroup.AppendInt(GGroomDataflowCardsLod);
+							
+				FString CardsFacesLODGroup = FGenerateCardsGeometryNode::CardsFacesGroup.ToString();
+				CardsFacesLODGroup.AppendInt(GGroomDataflowCardsLod);
+
+				if (Collection.HasAttribute(FGenerateCardsTexturesNode::VertexTextureUVsAttribute, FName(CardsVerticesLODGroup)) &&
+					Collection.HasAttribute(FGenerateCardsGeometryNode::FaceVertexIndicesAttribute, FName(CardsFacesLODGroup)))
+				{
+					const TManagedArray<FVector2f>& VertexUVsArray = Collection.GetAttribute<FVector2f>(
+							FGenerateCardsTexturesNode::VertexTextureUVsAttribute, FName(CardsVerticesLODGroup));
+						
+					const TManagedArray<FIntVector3>& FaceIndicesArray = Collection.GetAttribute<FIntVector3>(
+						FGenerateCardsGeometryNode::FaceVertexIndicesAttribute, FName(CardsFacesLODGroup));
+					
+					TArray<FIntVector> FaceIndices = FaceIndicesArray.GetConstArray();
+					
+					TArray<FVector3f> VertexUVs;
+					VertexUVs.Init(FVector3f(0.0, 0.0, 1.0), VertexUVsArray.Num());
+
+					for (int32 VertexIndex = 0; VertexIndex < VertexUVsArray.Num(); ++VertexIndex)
+					{
+						VertexUVs[VertexIndex] = FVector3f(VertexUVsArray[VertexIndex][0], VertexUVsArray[VertexIndex][1], 0.0f);
 					}
+
+					TArray<FVector3f> VertexNormals;
+					VertexNormals.Init(FVector3f(0.0, 0.0, 1.0), VertexUVs.Num());
+
+					TArray<FLinearColor> VertexColors;
+					VertexColors.Init(FLinearColor(0,0,0,0), VertexUVs.Num());
+					
+					const FString GeometryName = TEXT("Groom_UVs");
+
+					DynamicMeshes.SetNum(1);
+					BuildGeometryCollection(VertexUVs, VertexNormals, FaceIndices, VertexColors, FName(GeometryName), DynamicMeshes.Last());
 				}
 			}
+			AddMeshComponents(DynamicMeshes, OutComponents, {});
 		}
+	};
+
+	void RegisterCollectionRenderableTypes()
+	{
+		UE_DATAFLOW_REGISTER_RENDERABLE_TYPE(FGeometryCollectionCardsClumpsRenderableType);
+		UE_DATAFLOW_REGISTER_RENDERABLE_TYPE(FGeometryCollectionCardsGeometryRenderableType);
+		UE_DATAFLOW_REGISTER_RENDERABLE_TYPE(FGeometryCollectionCardsTextureRenderableType);
+		UE_DATAFLOW_REGISTER_RENDERABLE_TYPE(FGeometryCollectionCardsUVsRenderableType);
 	}
 }
+
+#undef LOCTEXT_NAMESPACE
diff --git a/Engine/Plugins/Experimental/HairCardGenerator/Source/HairCardGeneratorDataflow/Private/HairCardGeneratorDataflowModule.cpp b/Engine/Plugins/Experimental/HairCardGenerator/Source/HairCardGeneratorDataflow/Private/HairCardGeneratorDataflowModule.cpp
--- a/Engine/Plugins/Experimental/HairCardGenerator/Source/HairCardGeneratorDataflow/Private/HairCardGeneratorDataflowModule.cpp
+++ b/Engine/Plugins/Experimental/HairCardGenerator/Source/HairCardGeneratorDataflow/Private/HairCardGeneratorDataflowModule.cpp
@@ -19,15 +19,13 @@ void FHairCardGeneratorDataflowModule::StartupModule()
 	DATAFLOW_NODE_REGISTER_CREATION_FACTORY(FGenerateCardsTexturesNode);
 	DATAFLOW_NODE_REGISTER_CREATION_FACTORY(FBuildCardsSettingsNode);
 	DATAFLOW_NODE_REGISTER_CREATION_FACTORY(FCardsAssetTerminalNode);
-
-	UE::CardGen::Private::RegisterRenderingCallbacks();
-	UE::Groom::RegisterGroomDataflowTemplate({ "SelectCardsDataflowTemplate", "Cards Dataflow", "Add a Dataflow with a cards generation.", "/HairCardGenerator/Dataflow/DF_GroomCardsTemplate.DF_GroomCardsTemplate", false });
+	
+	UE::CardGen::Private::RegisterCollectionRenderableTypes();
+	UE::Groom::RegisterGroomDataflowTemplate({ "SelectCardsDataflowTemplate", "Cards", "Add a Dataflow with a cards generation.", "/HairCardGenerator/Dataflow/Templates/DF_GroomCardsTemplate.DF_GroomCardsTemplate", false });
 }
 
 void FHairCardGeneratorDataflowModule::ShutdownModule()
 {
-	UE::CardGen::Private::DeregisterRenderingCallbacks();
-
 	UE::Groom::UnregisterGroomDataflowTemplate("SelectCardsDataflowTemplate");
 }
 
diff --git a/Engine/Plugins/Experimental/HairCardGenerator/Source/HairCardGeneratorDataflow/Public/BuildCardsSettingsNode.h b/Engine/Plugins/Experimental/HairCardGenerator/Source/HairCardGeneratorDataflow/Public/BuildCardsSettingsNode.h
--- a/Engine/Plugins/Experimental/HairCardGenerator/Source/HairCardGeneratorDataflow/Public/BuildCardsSettingsNode.h
+++ b/Engine/Plugins/Experimental/HairCardGenerator/Source/HairCardGeneratorDataflow/Public/BuildCardsSettingsNode.h
@@ -50,6 +50,22 @@ struct FGroomFilterSettings
 	UPROPERTY(EditAnywhere, Category="Groom Cards")
 	int32 GroupIndex = 0;
 
+	// Total number of cards to generate for this LOD settings
+	UPROPERTY(EditAnywhere, Category="Groom Cards", meta=(ClampMin="1", ClampMax="100000", DisplayName = "Num Clumps"))
+	int32 NumClumps = 50;
+
+	// Maximum number of cards to assign to flyaway strands
+	UPROPERTY(EditAnywhere, Category="Groom Cards", meta=(ClampMin="0", ClampMax="1000"))
+	int32 NumFlyaways = 10;
+
+	// Total number of triangles to generate for this LOD settings
+	UPROPERTY(EditAnywhere, Category="Groom Cards", meta=(ClampMin="1", ClampMax="100000"))
+	int32 NumTriangles = 2000;
+
+	// Total number of triangles to generate for this LOD settings
+	UPROPERTY(EditAnywhere, Category="Groom Cards")
+	int32 NumTextures = 75;
+
 	// Card group names that will belong to the filter group settings
 	UPROPERTY(EditAnywhere, Category="Groom Cards")
 	TArray<FName> CardGroups;
@@ -65,13 +81,9 @@ struct FBuildCardsSettingsNode : public FDataflowNode
 
 public:
 	
-	FBuildCardsSettingsNode(const UE::Dataflow::FNodeParameters& InParam, FGuid InGuid = FGuid::NewGuid())
-	: FDataflowNode(InParam, InGuid)
-	{
-		RegisterInputConnection(&Collection);
-		RegisterOutputConnection(&Collection, &Collection);
-		RegisterOutputConnection(&CardsSettings);
-	}
+	FBuildCardsSettingsNode(const UE::Dataflow::FNodeParameters& InParam, FGuid InGuid = FGuid::NewGuid());
+
+private:
 
 	//~ Begin FDataflowNode interface
 	virtual void Evaluate(UE::Dataflow::FContext& Context, const FDataflowOutput* Out) const override;
@@ -83,7 +95,7 @@ public:
 	//~ End FDataflowNode interface
 
 	/** Groom asset to build the cards settings from */
-	UPROPERTY(EditAnywhere, Category = "Groom Cards")
+	UPROPERTY(EditAnywhere, Category = "Groom Cards", meta = (DataflowInput))
 	TObjectPtr<UGroomAsset> GroomAsset;
 	
 	/** Managed array collection to be used to store datas */
diff --git a/Engine/Plugins/Experimental/HairCardGenerator/Source/HairCardGeneratorDataflow/Public/CardsAssetTerminalNode.h b/Engine/Plugins/Experimental/HairCardGenerator/Source/HairCardGeneratorDataflow/Public/CardsAssetTerminalNode.h
--- a/Engine/Plugins/Experimental/HairCardGenerator/Source/HairCardGeneratorDataflow/Public/CardsAssetTerminalNode.h
+++ b/Engine/Plugins/Experimental/HairCardGenerator/Source/HairCardGeneratorDataflow/Public/CardsAssetTerminalNode.h
@@ -19,13 +19,9 @@ struct FCardsAssetTerminalNode : public FDataflowTerminalNode
 
 public:
 
-	FCardsAssetTerminalNode(const UE::Dataflow::FNodeParameters& InParam, FGuid InGuid = FGuid::NewGuid())
-	: FDataflowTerminalNode(InParam, InGuid)
-	{
-		RegisterInputConnection(&Collection);
-		RegisterInputConnection(&CardsSettings);
-		RegisterOutputConnection(&Collection, &Collection);
-	}
+	FCardsAssetTerminalNode(const UE::Dataflow::FNodeParameters& InParam, FGuid InGuid = FGuid::NewGuid());
+	
+private:
 
 	//~ Begin FDataflowTerminalNode interface
 	virtual void SetAssetValue(TObjectPtr<UObject> Asset, UE::Dataflow::FContext& Context) const override;
diff --git a/Engine/Plugins/Experimental/HairCardGenerator/Source/HairCardGeneratorDataflow/Public/GenerateCardsClumpsNode.h b/Engine/Plugins/Experimental/HairCardGenerator/Source/HairCardGeneratorDataflow/Public/GenerateCardsClumpsNode.h
--- a/Engine/Plugins/Experimental/HairCardGenerator/Source/HairCardGeneratorDataflow/Public/GenerateCardsClumpsNode.h
+++ b/Engine/Plugins/Experimental/HairCardGenerator/Source/HairCardGeneratorDataflow/Public/GenerateCardsClumpsNode.h
@@ -36,32 +36,27 @@ struct FGenerateCardsClumpsNode : public FDataflowNode
 	GENERATED_BODY()
 	
 	DATAFLOW_NODE_DEFINE_INTERNAL(FGenerateCardsClumpsNode, "GenerateCardsClumps", "Groom", "")
-	DATAFLOW_NODE_RENDER_TYPE("ClumpsRender", FName("FCardsCollection"), "Collection")
 
 public:
 
 	static const FName CurveClumpIndicesAttribute;
 	static const FName ObjectNumClumpsAttribute;
+	static const FName CurveFilterIndicesAttribute;
 	
-	FGenerateCardsClumpsNode(const UE::Dataflow::FNodeParameters& InParam, FGuid InGuid = FGuid::NewGuid())
-	: FDataflowNode(InParam, InGuid)
-	{
-		RegisterInputConnection(&Collection);
-		RegisterInputConnection(&CardsSettings);
-		RegisterOutputConnection(&Collection, &Collection);
-		RegisterOutputConnection(&CardsSettings, &CardsSettings);
-	}
+	FGenerateCardsClumpsNode(const UE::Dataflow::FNodeParameters& InParam, FGuid InGuid = FGuid::NewGuid());
+
+private:
 
 	//~ Begin FDataflowNode interface
 	virtual void Evaluate(UE::Dataflow::FContext& Context, const FDataflowOutput* Out) const override;
 	//~ End FDataflowNode interface
 	
 	/** Per LOD settings for each cards generation */
-	UPROPERTY(EditAnywhere, Category = "GroomCards")
+	UPROPERTY(EditAnywhere, Category = "Groom Cards", meta = (DisplayName = "Override Settings"))
 	TArray<FCardsClumpsSettings> ClumpsSettings;
 
 	/** Managed array collection to be used to store datas */
-	UPROPERTY(meta = (DataflowInput, DataflowOutput, DisplayName = "Collection", DataflowPassthrough  = "Collection"))
+	UPROPERTY(meta = (DataflowInput, DataflowOutput, DisplayName = "Collection", DataflowPassthrough  = "Collection", DataflowRenderGroups = "CardsClumps"))
 	FManagedArrayCollection Collection;
 
 	/** Generator settings to be used */
diff --git a/Engine/Plugins/Experimental/HairCardGenerator/Source/HairCardGeneratorDataflow/Public/GenerateCardsGeometryNode.h b/Engine/Plugins/Experimental/HairCardGenerator/Source/HairCardGeneratorDataflow/Public/GenerateCardsGeometryNode.h
--- a/Engine/Plugins/Experimental/HairCardGenerator/Source/HairCardGeneratorDataflow/Public/GenerateCardsGeometryNode.h
+++ b/Engine/Plugins/Experimental/HairCardGenerator/Source/HairCardGeneratorDataflow/Public/GenerateCardsGeometryNode.h
@@ -32,7 +32,6 @@ struct FGenerateCardsGeometryNode : public FDataflowNode
 	GENERATED_BODY()
 	
 	DATAFLOW_NODE_DEFINE_INTERNAL(FGenerateCardsGeometryNode, "GenerateCardsGeometry", "Groom", "")
-	DATAFLOW_NODE_RENDER_TYPE("GeometryRender", FName("FCardsCollection"), "Collection")
 
 public:
 
@@ -43,29 +42,24 @@ public:
 	static const FName CardsVerticesGroup;
 	static const FName CardsFacesGroup;
 	
-	FGenerateCardsGeometryNode(const UE::Dataflow::FNodeParameters& InParam, FGuid InGuid = FGuid::NewGuid())
-	: FDataflowNode(InParam, InGuid)
-	{
-		RegisterInputConnection(&Collection);
-		RegisterInputConnection(&CardsSettings);
-		RegisterOutputConnection(&Collection, &Collection);
-		RegisterOutputConnection(&CardsSettings, &CardsSettings);
-	}
+	FGenerateCardsGeometryNode(const UE::Dataflow::FNodeParameters& InParam, FGuid InGuid = FGuid::NewGuid());
+
+private:
 
 	//~ Begin FDataflowNode interface
 	virtual void Evaluate(UE::Dataflow::FContext& Context, const FDataflowOutput* Out) const override;
 	//~ End FDataflowNode interface
 
 	/** Per LOD settings for each cards generation */
-	UPROPERTY(EditAnywhere, Category = "GroomCards")
+	UPROPERTY(EditAnywhere, Category = "Groom Cards", meta = (DisplayName = "Override Settings"))
 	TArray<FCardsGeometrySettings> GeometrySettings;
 
 	/** Managed array collection to be used to store datas */
-	UPROPERTY(meta = (DataflowInput, DataflowOutput, DisplayName = "Collection", DataflowPassthrough  = "Collection"))
+	UPROPERTY(meta = (DataflowInput, DataflowOutput, DisplayName = "Collection", DataflowPassthrough  = "Collection", DataflowRenderGroups = "CardsGeometry"))
 	FManagedArrayCollection Collection;
 
 	/** Generator settings to be used */
-	UPROPERTY(meta = (DataflowInput, DataflowOutput, DisplayName = "CardsSettings", DataflowPassthrough  = "CardsSettings"));
+	UPROPERTY(meta = (DataflowInput, DataflowOutput, DisplayName = "Cards Settings", DataflowPassthrough  = "CardsSettings"));
 	TArray<FGroomCardsSettings> CardsSettings;
 };
 
diff --git a/Engine/Plugins/Experimental/HairCardGenerator/Source/HairCardGeneratorDataflow/Public/GenerateCardsTexturesNode.h b/Engine/Plugins/Experimental/HairCardGenerator/Source/HairCardGeneratorDataflow/Public/GenerateCardsTexturesNode.h
--- a/Engine/Plugins/Experimental/HairCardGenerator/Source/HairCardGeneratorDataflow/Public/GenerateCardsTexturesNode.h
+++ b/Engine/Plugins/Experimental/HairCardGenerator/Source/HairCardGeneratorDataflow/Public/GenerateCardsTexturesNode.h
@@ -25,14 +25,13 @@ struct FCardsTextureSettings
 	int32 NumTextures = 75;
 };
 
-/** Generate the clumps used to build the cards from the strands */
+/** Generate the cards texture for final rendering */
 USTRUCT(meta = (Experimental, DataflowGroom))
 struct FGenerateCardsTexturesNode : public FDataflowNode
 {
 	GENERATED_BODY()
 	
 	DATAFLOW_NODE_DEFINE_INTERNAL(FGenerateCardsTexturesNode, "GenerateCardsTextures", "Groom", "")
-	DATAFLOW_NODE_RENDER_TYPE("TextureRender", FName("FCardsCollection"), "Collection")
 
 public:
 
@@ -41,29 +40,24 @@ public:
 	static const FName ObjectTextureIndicesAttribute;
 	static const FName VertexTextureUVsAttribute;
 	
-	FGenerateCardsTexturesNode(const UE::Dataflow::FNodeParameters& InParam, FGuid InGuid = FGuid::NewGuid())
-	: FDataflowNode(InParam, InGuid)
-	{
-		RegisterInputConnection(&Collection);
-		RegisterInputConnection(&CardsSettings);
-		RegisterOutputConnection(&Collection, &Collection);
-		RegisterOutputConnection(&CardsSettings, &CardsSettings);
-	}
+	FGenerateCardsTexturesNode(const UE::Dataflow::FNodeParameters& InParam, FGuid InGuid = FGuid::NewGuid());
+
+private:
 
 	//~ Begin FDataflowNode interface
 	virtual void Evaluate(UE::Dataflow::FContext& Context, const FDataflowOutput* Out) const override;
 	//~ End FDataflowNode interface
 
 	/** Per LOD settings for each cards generation */
-	UPROPERTY(EditAnywhere, Category = "GroomCards")
+	UPROPERTY(EditAnywhere, Category = "Groom Cards", meta = (DisplayName = "Override Settings"))
 	TArray<FCardsTextureSettings> TextureSettings;
 
 	/** Managed array collection to be used to store datas */
-	UPROPERTY(meta = (DataflowInput, DataflowOutput, DisplayName = "Collection", DataflowPassthrough  = "Collection"))
+	UPROPERTY(meta = (DataflowInput, DataflowOutput, DisplayName = "Collection", DataflowPassthrough  = "Collection", DataflowRenderGroups = "CardsTexture"))
 	FManagedArrayCollection Collection;
 
 	/** Generator settings to be used */
-	UPROPERTY(meta = (DataflowInput, DataflowOutput, DisplayName = "CardsSettings", DataflowPassthrough  = "CardsSettings"));
+	UPROPERTY(meta = (DataflowInput, DataflowOutput, DisplayName = "Cards Settings", DataflowPassthrough  = "CardsSettings"));
 	TArray<FGroomCardsSettings> CardsSettings;
 };
 
diff --git a/Engine/Plugins/Experimental/HairCardGenerator/Source/HairCardGeneratorDataflow/Public/HairCardDataflowRendering.h b/Engine/Plugins/Experimental/HairCardGenerator/Source/HairCardGeneratorDataflow/Public/HairCardDataflowRendering.h
--- a/Engine/Plugins/Experimental/HairCardGenerator/Source/HairCardGeneratorDataflow/Public/HairCardDataflowRendering.h
+++ b/Engine/Plugins/Experimental/HairCardGenerator/Source/HairCardGeneratorDataflow/Public/HairCardDataflowRendering.h
@@ -2,85 +2,7 @@
 
 #pragma once
 
-#include "GroomDataflowRendering.h"
-#include "Dataflow/DataflowRenderingFactory.h"
-#include "Dataflow/DataflowRenderingViewMode.h"
-
 namespace UE::CardGen::Private
 {
-	class FCardsGeometryRenderingCallbacks : public Dataflow::FRenderingFactory::ICallbackInterface
-	{
-	public :
-		static Dataflow::FRenderKey RenderKey;
-		
-	protected:
-
-		virtual Dataflow::FRenderKey GetRenderKey() const override
-		{
-			return RenderKey;
-		}
-
-		virtual bool CanRender(const Dataflow::IDataflowConstructionViewMode& ViewMode) const override
-		{
-			return (ViewMode.GetName() == Dataflow::FDataflowConstruction3DViewMode::Name);
-		}
-
-		virtual void Render(GeometryCollection::Facades::FRenderingFacade& RenderCollection, const Dataflow::FGraphRenderingState& State) override;
-		virtual void ComputeVertexColors(const FManagedArrayCollection& Collection, const int32 CardIndex, TArray<FLinearColor>& VertexColors) const {}
-	};
-	
-	class FCardsClumpsRenderingCallbacks : public Dataflow::FRenderingFactory::ICallbackInterface
-	{
-		public :
-			static Dataflow::FRenderKey RenderKey;
-		
-	protected:
-
-		virtual Dataflow::FRenderKey GetRenderKey() const override
-		{
-			return RenderKey;
-		}
-
-		virtual bool CanRender(const Dataflow::IDataflowConstructionViewMode& ViewMode) const override
-		{
-			return (ViewMode.GetName() == Dataflow::FDataflowConstruction3DViewMode::Name);
-		}
-
-		virtual void Render(GeometryCollection::Facades::FRenderingFacade& RenderCollection, const Dataflow::FGraphRenderingState& State) override;
-	};
-
-	class FCardsTextureRenderingCallbacks : public FCardsGeometryRenderingCallbacks
-	{
-	public :
-			static Dataflow::FRenderKey RenderKey;
-		
-	private:
-
-		virtual bool CanRender(const Dataflow::IDataflowConstructionViewMode& ViewMode) const override
-		{
-			return (ViewMode.GetName() == Dataflow::FDataflowConstruction3DViewMode::Name || ViewMode.GetName() == Dataflow::FDataflowConstructionUVViewMode::Name);
-		}
-
-		virtual Dataflow::FRenderKey GetRenderKey() const override
-		{
-			return RenderKey;
-		}
-		virtual void Render(GeometryCollection::Facades::FRenderingFacade& RenderCollection, const Dataflow::FGraphRenderingState& State) override;
-		virtual void ComputeVertexColors(const FManagedArrayCollection& Collection, const int32 CardIndex, TArray<FLinearColor>& VertexColors) const override;
-	};
-	
-	static void RegisterRenderingCallbacks()
-	{
-		Dataflow::FRenderingFactory::GetInstance()->RegisterCallbacks(MakeUnique<FCardsClumpsRenderingCallbacks>());
-		Dataflow::FRenderingFactory::GetInstance()->RegisterCallbacks(MakeUnique<FCardsGeometryRenderingCallbacks>());
-		Dataflow::FRenderingFactory::GetInstance()->RegisterCallbacks(MakeUnique<FCardsTextureRenderingCallbacks>());
-	}
-
-	static void DeregisterRenderingCallbacks()
-	{
-		Dataflow::FRenderingFactory::GetInstance()->DeregisterCallbacks(FCardsGeometryRenderingCallbacks::RenderKey);
-		Dataflow::FRenderingFactory::GetInstance()->DeregisterCallbacks(FCardsClumpsRenderingCallbacks::RenderKey);
-		Dataflow::FRenderingFactory::GetInstance()->DeregisterCallbacks(FCardsTextureRenderingCallbacks::RenderKey);
-	}
-
+	void RegisterCollectionRenderableTypes();
 }  // End namespace Private
diff --git a/Engine/Plugins/Experimental/HairCardGenerator/Source/HairCardGeneratorEditor/Private/HairCardGeneratorEditorModule.cpp b/Engine/Plugins/Experimental/HairCardGenerator/Source/HairCardGeneratorEditor/Private/HairCardGeneratorEditorModule.cpp
--- a/Engine/Plugins/Experimental/HairCardGenerator/Source/HairCardGeneratorEditor/Private/HairCardGeneratorEditorModule.cpp
+++ b/Engine/Plugins/Experimental/HairCardGenerator/Source/HairCardGeneratorEditor/Private/HairCardGeneratorEditorModule.cpp
@@ -409,9 +409,9 @@ static bool GenerateHairCardGenGroomData(const int PointsPerStrand, const FHairD
 	return true;
 }
 
-void FHairCardGeneratorUtils::BuildCardsGroups(TObjectPtr<const UGroomAsset> GroomAsset, TArray<FName>& GroupNames)
+void FHairCardGeneratorUtils::BuildCardsGroups(TObjectPtr<const UGroomAsset> GroomAsset, TArray<FName>& GroupNames, TArray<int32>& GroupCount)
 {
-	HairCardSettings_Helpers::GetStrandCardGroupInfo(GroomAsset, GroupNames);
+	HairCardSettings_Helpers::GetStrandCardGroupInfo(GroomAsset, GroupNames, GroupCount);
 }
 
 bool FHairCardGeneratorUtils::LoadGroomStrands(const UGroomAsset* GroomAsset, const FBuilderFunction& BuilderFunction)
diff --git a/Engine/Plugins/Experimental/HairCardGenerator/Source/HairCardGeneratorEditor/Private/HairCardGeneratorPluginSettings.cpp b/Engine/Plugins/Experimental/HairCardGenerator/Source/HairCardGeneratorEditor/Private/HairCardGeneratorPluginSettings.cpp
--- a/Engine/Plugins/Experimental/HairCardGenerator/Source/HairCardGeneratorEditor/Private/HairCardGeneratorPluginSettings.cpp
+++ b/Engine/Plugins/Experimental/HairCardGenerator/Source/HairCardGeneratorEditor/Private/HairCardGeneratorPluginSettings.cpp
@@ -233,9 +233,9 @@ bool HairCardSettings_Helpers::IsWriteableSettingsProperty(FProperty* Property)
 	return Property->HasAllPropertyFlags(CPF_BlueprintVisible);
 }
 
-void HairCardSettings_Helpers::GetStrandCardGroupInfo(TObjectPtr<const UGroomAsset> GroomAsset, TArray<FName>& OutGrpIds)
+void HairCardSettings_Helpers::GetStrandCardGroupInfo(TObjectPtr<const UGroomAsset> GroomAsset, TArray<FName>& OutGrpIds, TArray<int32>& OutGrpCount)
 {
-    TSet<FName> IdNames;
+    TMap<FName,int32> IdNames;
 
     FHairDescription HairDesc = GroomAsset->GetHairDescription();
 	FStrandAttributesRange TraverseStrands(HairDesc, FHairStrandAttribsConstIterator::IgnoreNone);
@@ -243,7 +243,15 @@ void HairCardSettings_Helpers::GetStrandCardGroupInfo(TObjectPtr<const UGroomAss
     {
 		for ( const FHairStrandAttribsRefProxy& StrandAttribs : TraverseStrands )
 		{
-			IdNames.Add(StrandAttribs.GetCardGroupName());
+			int32* NumStrands = IdNames.Find(StrandAttribs.GetCardGroupName());
+			if(!NumStrands)
+			{
+				IdNames.Add(StrandAttribs.GetCardGroupName(), 0);
+			}
+			else
+			{
+				(*NumStrands) += 1;
+			}
 		}
 
         // IdCountMap.KeySort([](FName a, FName b){return a < b;});
@@ -255,7 +263,8 @@ void HairCardSettings_Helpers::GetStrandCardGroupInfo(TObjectPtr<const UGroomAss
 		UE_LOG(LogHairCardGenerator, Warning, TEXT("Groom has groom_group_cards_id attribute, but some strands are empty (will be ignored)"));
 	}
 
-    OutGrpIds = IdNames.Array();
+	IdNames.GenerateKeyArray(OutGrpIds);
+	IdNames.GenerateValueArray(OutGrpCount);
 }
 
 
@@ -341,7 +350,8 @@ void UHairCardGeneratorPluginSettings::ResetFromSourceDescription(const FHairGro
 		DestinationPath.Path = HairCardSettings_Helpers::GetDefaultPathForHairCards(GroomAsset, LODIndex);
 	}
 
-	HairCardSettings_Helpers::GetStrandCardGroupInfo(GroomAsset, CardGroupIds);
+	TArray<int32> CardGroupCount;
+	HairCardSettings_Helpers::GetStrandCardGroupInfo(GroomAsset, CardGroupIds, CardGroupCount);
 
 	int GroupSettingsCount = FMath::Max(1, LastSettings->FilterGroupGenerationSettings.Num());
 	for ( int i=0; i < GroupSettingsCount; ++i )
@@ -393,7 +403,8 @@ void UHairCardGeneratorPluginSettings::ResetToDefault()
 	BaseFilename = HairCardSettings_Helpers::GetDefaultBaseNameForHairCards(GroomAsset, LODIndex);
 	DestinationPath.Path = HairCardSettings_Helpers::GetDefaultPathForHairCards(GroomAsset, LODIndex);
 
-	HairCardSettings_Helpers::GetStrandCardGroupInfo(GroomAsset, CardGroupIds);
+	TArray<int32> CardGroupCount;
+	HairCardSettings_Helpers::GetStrandCardGroupInfo(GroomAsset, CardGroupIds, CardGroupCount);
 
 	ResetFilterGroupSettings(0, UHairCardGeneratorGroupSettings::StaticClass()->GetDefaultObject<UHairCardGeneratorGroupSettings>());
 	ResetNumFilterGroups(1);
diff --git a/Engine/Plugins/Experimental/HairCardGenerator/Source/HairCardGeneratorEditor/Private/HairCardGeneratorPluginSettings.h b/Engine/Plugins/Experimental/HairCardGenerator/Source/HairCardGeneratorEditor/Private/HairCardGeneratorPluginSettings.h
--- a/Engine/Plugins/Experimental/HairCardGenerator/Source/HairCardGeneratorEditor/Private/HairCardGeneratorPluginSettings.h
+++ b/Engine/Plugins/Experimental/HairCardGenerator/Source/HairCardGeneratorEditor/Private/HairCardGeneratorPluginSettings.h
@@ -12,7 +12,7 @@ class FJsonObject;
 
 namespace HairCardSettings_Helpers
 {
-	void GetStrandCardGroupInfo(TObjectPtr<const UGroomAsset> GroomAsset, TArray<FName>& OutGrpIds);
+	void GetStrandCardGroupInfo(TObjectPtr<const UGroomAsset> GroomAsset, TArray<FName>& OutGrpIds, TArray<int32>& OutGrpCount);
 };
 
 UENUM(BlueprintType)
diff --git a/Engine/Plugins/Experimental/HairCardGenerator/Source/HairCardGeneratorEditor/Public/HairCardGeneratorEditorModule.h b/Engine/Plugins/Experimental/HairCardGenerator/Source/HairCardGeneratorEditor/Public/HairCardGeneratorEditorModule.h
--- a/Engine/Plugins/Experimental/HairCardGenerator/Source/HairCardGeneratorEditor/Public/HairCardGeneratorEditorModule.h
+++ b/Engine/Plugins/Experimental/HairCardGenerator/Source/HairCardGeneratorEditor/Public/HairCardGeneratorEditorModule.h
@@ -45,7 +45,7 @@ struct HAIRCARDGENERATOREDITOR_API FHairCardGeneratorUtils
 	static bool LoadGroomStrands(const UGroomAsset* GroomAsset, const FBuilderFunction& BuilderFunction);
 	
 	/** Build the cards group names */
-	static void BuildCardsGroups(TObjectPtr<const UGroomAsset> GroomAsset, TArray<FName>& GroupNames);
+	static void BuildCardsGroups(TObjectPtr<const UGroomAsset> GroomAsset, TArray<FName>& GroupNames, TArray<int32>& GroupCount);
 };
 
 /** FHairCardGeneratorEditorModule  */
diff --git a/Engine/Plugins/Runtime/HairStrands/Source/HairStrandsCore/Public/GroomAssetPhysics.h b/Engine/Plugins/Runtime/HairStrands/Source/HairStrandsCore/Public/GroomAssetPhysics.h
--- a/Engine/Plugins/Runtime/HairStrands/Source/HairStrandsCore/Public/GroomAssetPhysics.h
+++ b/Engine/Plugins/Runtime/HairStrands/Source/HairStrandsCore/Public/GroomAssetPhysics.h
@@ -51,11 +51,11 @@ struct FHairSolverSettings
 	UE_API FHairSolverSettings();
 	
 	/** Use deformer for simulation  */
-	UPROPERTY(EditAnywhere, Category = "SolverSettings", meta = (ToolTip = "Use deformer for the groom solver", DisplayName = "Enable Deformer Simulation"))
+	UPROPERTY(EditAnywhere, Category = "SolverSettings", meta = (ToolTip = "Use deformer for the groom solver", DisplayName = "Enable Deformation"))
 	bool bEnableDeformation;
 
 	/** Use niagara for simulation */
-	UPROPERTY(EditAnywhere, Category = "SolverSettings", meta = (ToolTip = "Use niagara for the groom solver", DisplayName = "Enable Niagara Simulation"))
+	UPROPERTY(EditAnywhere, Category = "SolverSettings", meta = (ToolTip = "Use niagara for the groom solver", DisplayName = "Enable Simulation"))
 	bool EnableSimulation;
 
 	/** Niagara solver to be used for simulation */
