Dataflow : fix performance issue when calling SetShadowEnabled  on the dynamic mesh component
- Replace SetShadowEnabled  with SetCastShadow to avoid flushing the rendering commands all the time

diff --git a/Engine/Plugins/Experimental/Dataflow/Source/DataflowEditor/Private/Dataflow/DataflowConstructionScene.cpp b/Engine/Plugins/Experimental/Dataflow/Source/DataflowEditor/Private/Dataflow/DataflowConstructionScene.cpp
--- a/Engine/Plugins/Experimental/Dataflow/Source/DataflowEditor/Private/Dataflow/DataflowConstructionScene.cpp
+++ b/Engine/Plugins/Experimental/Dataflow/Source/DataflowEditor/Private/Dataflow/DataflowConstructionScene.cpp
@@ -280,7 +280,7 @@ TObjectPtr<UDynamicMeshComponent>& FDataflowConstructionScene::AddDynamicMeshCom
 		{
 			ensure(FDataflowEditorStyle::Get().DefaultTwoSidedMaterial);
 			DynamicMeshComponent->SetOverrideRenderMaterial(FDataflowEditorStyle::Get().DefaultTwoSidedMaterial);
-			DynamicMeshComponent->SetShadowsEnabled(false);
+			DynamicMeshComponent->SetCastShadow(false);
 		}
 	}
 	else
diff --git a/Engine/Plugins/Experimental/Dataflow/Source/DataflowEditor/Private/DataflowRendering/DataflowGeometryCollectionRenderableType.cpp b/Engine/Plugins/Experimental/Dataflow/Source/DataflowEditor/Private/DataflowRendering/DataflowGeometryCollectionRenderableType.cpp
--- a/Engine/Plugins/Experimental/Dataflow/Source/DataflowEditor/Private/DataflowRendering/DataflowGeometryCollectionRenderableType.cpp
+++ b/Engine/Plugins/Experimental/Dataflow/Source/DataflowEditor/Private/DataflowRendering/DataflowGeometryCollectionRenderableType.cpp
@@ -86,9 +86,9 @@ namespace UE::Dataflow::Private
 								FName(TEXT("Bone_Dynamic_Mesh_") + FString::FromInt(TransformIndex));
 						if (UDynamicMeshComponent* Component = OutComponents.AddNewComponent<UDynamicMeshComponent>(ComponentName))
 						{
+							Component->SetCastShadow(false);
 							Component->SetMesh(MoveTemp(DynamicMeshes[TransformIndex]));
 							Component->SetOverrideRenderMaterial(FDataflowEditorStyle::Get().VertexMaterial);
-							Component->SetShadowsEnabled(false);
 						}
 					}
 				}
diff --git a/Engine/Plugins/Experimental/HairCardGenerator/Source/HairCardGeneratorDataflow/Private/HairCardDataflowRendering.cpp b/Engine/Plugins/Experimental/HairCardGenerator/Source/HairCardGeneratorDataflow/Private/HairCardDataflowRendering.cpp
--- a/Engine/Plugins/Experimental/HairCardGenerator/Source/HairCardGeneratorDataflow/Private/HairCardDataflowRendering.cpp
+++ b/Engine/Plugins/Experimental/HairCardGenerator/Source/HairCardGeneratorDataflow/Private/HairCardDataflowRendering.cpp
@@ -140,7 +140,7 @@ namespace UE::CardGen::Private
 								Component->SetupAttachment(FilterComponent);
 								Component->SetMesh(MoveTemp(DynamicMeshes[TransformIndex]));
 								Component->SetOverrideRenderMaterial(Dataflow::Rendering::GetVertexMaterial());
-								Component->SetShadowsEnabled(false);
+								Component->SetCastShadow(false);
 							}
 						}
 						++LocalIndex;
@@ -159,7 +159,7 @@ namespace UE::CardGen::Private
 					{
 						Component->SetMesh(MoveTemp(DynamicMeshes[TransformIndex]));
 						Component->SetOverrideRenderMaterial(Dataflow::Rendering::GetVertexMaterial());
-						Component->SetShadowsEnabled(false);
+						Component->SetCastShadow(false);
 					}
 				}
 			}
