Fix card texture rendering + add generation settings + automatic cardsgroups creation

diff --git a/Engine/Plugins/Experimental/HairCardGenerator/Source/HairCardGeneratorDataflow/Private/BuildCardsSettingsNode.cpp b/Engine/Plugins/Experimental/HairCardGenerator/Source/HairCardGeneratorDataflow/Private/BuildCardsSettingsNode.cpp
--- a/Engine/Plugins/Experimental/HairCardGenerator/Source/HairCardGeneratorDataflow/Private/BuildCardsSettingsNode.cpp
+++ b/Engine/Plugins/Experimental/HairCardGenerator/Source/HairCardGeneratorDataflow/Private/BuildCardsSettingsNode.cpp
@@ -70,66 +70,97 @@ void FBuildCardsSettingsNode::Evaluate(UE::Dataflow::FContext& Context, const FD
 	}
 	else if (Out->IsA<TArray<FGroomCardsSettings>>(&CardsSettings))
 	{
+		const FManagedArrayCollection& GeometryCollection = GetValue<FManagedArrayCollection>(Context, &Collection);
 		UGroomAsset* LocalAsset = GetValue<TObjectPtr<UGroomAsset>>(Context, &GroomAsset);
 		TArray<FGroomCardsSettings> OutputSettings;
-		if(LocalAsset && !LocalAsset->GetHairGroupsCards().IsEmpty())
+		if (LocalAsset)
 		{
-			FManagedArrayCollection GeometryCollection = GetValue<FManagedArrayCollection>(Context, &Collection);
-			for(FHairGroupsCardsSourceDescription& HairCardsDescription : LocalAsset->GetHairGroupsCards())
+			for(const FGroomGenerationSettings& LocalSettings : GenerationSettings)
 			{
-				if(HairCardsDescription.LODIndex != INDEX_NONE)
-				{ 
-					HairCardsDescription.Textures = FHairGroupCardsTextures();
+				FHairGroupsCardsSourceDescription* CardsDescriptionPtr = nullptr;
 				
-					OutputSettings.AddDefaulted();
-					FHairCardGeneratorUtils::BuildGenerationSettings(false, LocalAsset, HairCardsDescription,
-						OutputSettings.Last().GenerationSettings, OutputSettings.Last().GenerationFlags, OutputSettings.Last().PipelineFlags);
+				bool bHasCardDescription = false;
+				for(FHairGroupsCardsSourceDescription& HairCardsDescription : LocalAsset->GetHairGroupsCards())
+				{
+					if((LocalSettings.GroupIndex == HairCardsDescription.GroupIndex) &&
+								(LocalSettings.LODIndex == HairCardsDescription.LODIndex) && (HairCardsDescription.LODIndex != INDEX_NONE))
+					{
+						bHasCardDescription = true;
+						CardsDescriptionPtr = &HairCardsDescription;
+						break;
+					}
+				}
+				if (!bHasCardDescription)
+				{
+					FHairGroupsCardsSourceDescription CardsDescription;
+					CardsDescription.LODIndex = LocalSettings.LODIndex;
+					CardsDescription.GroupIndex = LocalSettings.GroupIndex;
+					LocalAsset->GetHairGroupsCards().Add(CardsDescription);
+
+					FHairGroupsCardsSourceDescription& HairCardsDescription = LocalAsset->GetHairGroupsCards().Last();
+					CardsDescriptionPtr = &HairCardsDescription;
+				}
+
+				if (CardsDescriptionPtr)
+				{
+					CardsDescriptionPtr->Textures = FHairGroupCardsTextures();
+
+					FCardsGenerationAdvancedOptions AdvancedOptions;
+					AdvancedOptions.bGenerateGeometryForAllGroups = LocalSettings.AdvancedOptions.bGenerateGeometryForAllGroups;
+					AdvancedOptions.bReduceCardsFromPreviousLOD = LocalSettings.AdvancedOptions.bReduceCardsFromPreviousLOD;
+					AdvancedOptions.RandomSeed = LocalSettings.AdvancedOptions.RandomSeed;
+					AdvancedOptions.bUseReservedSpaceFromPreviousLOD = LocalSettings.AdvancedOptions.bUseReservedSpaceFromPreviousLOD;
+					AdvancedOptions.AtlasSize = static_cast<uint8>(LocalSettings.AdvancedOptions.AtlasSize);
+					AdvancedOptions.ReserveTextureSpaceLOD = LocalSettings.AdvancedOptions.ReserveTextureSpaceLOD;
+					AdvancedOptions.bUseGroomAssetStrandWidth = LocalSettings.AdvancedOptions.bUseGroomAssetStrandWidth;
 				
-					if(!FilterSettings.IsEmpty())
+					OutputSettings.AddDefaulted();
+					FHairCardGeneratorUtils::BuildGenerationSettings(false, LocalAsset, *CardsDescriptionPtr,
+						OutputSettings.Last().GenerationSettings, OutputSettings.Last().GenerationFlags, OutputSettings.Last().PipelineFlags, AdvancedOptions);
+
+					if(!LocalSettings.FilterSettings.IsEmpty())
 					{
 						TArray<TArray<FName>> FilterCardGroups;
 						TArray<FName> FilterGroupNames;
-						for(const FGroomFilterSettings& LocalSettings : FilterSettings)
+						for(const FGroomFilterSettings& FilterlSettings : LocalSettings.FilterSettings)
 						{
-							if((LocalSettings.GroupIndex == HairCardsDescription.GroupIndex) &&
-								(LocalSettings.LODIndex == HairCardsDescription.LODIndex) && (HairCardsDescription.LODIndex != INDEX_NONE))
-							{
-								FilterCardGroups.Add(LocalSettings.CardGroups);
-								FilterGroupNames.Add(LocalSettings.FilterName);
-							}
+							FilterCardGroups.Add(FilterlSettings.CardGroups);
+							FilterGroupNames.Add(FilterlSettings.FilterName);
 						}
 						OutputSettings.Last().GenerationSettings->BuildFilterGroupSettings(FilterCardGroups, FilterGroupNames);
 	
-						for(const FGroomFilterSettings& LocalSettings : FilterSettings)
+						for(const FGroomFilterSettings& FilterlSettings : LocalSettings.FilterSettings)
 						{
 							for(TObjectPtr<UHairCardGeneratorGroupSettings> FilterGroups : OutputSettings.Last().GenerationSettings->GetFilterGroupSettings())
 							{
 								if (FilterGroups.Get())
 								{
-									if ((LocalSettings.FilterName != NAME_None && FilterGroups->GetFilterName() == LocalSettings.FilterName) || LocalSettings.FilterName == NAME_None)
+									if ((FilterlSettings.FilterName != NAME_None && FilterGroups->GetFilterName() == FilterlSettings.FilterName) || FilterlSettings.FilterName == NAME_None)
 									{
-										FilterGroups->TargetNumberOfCards = LocalSettings.NumClumps * 3;
-										FilterGroups->MaxNumberOfFlyaways = LocalSettings.NumFlyaways;
-										FilterGroups->TargetTriangleCount = LocalSettings.NumTriangles;
-										FilterGroups->NumberOfTexturesInAtlas = LocalSettings.NumTextures;
+										FilterGroups->TargetNumberOfCards = FilterlSettings.NumClumps * 3;
+										FilterGroups->MaxNumberOfFlyaways = FilterlSettings.NumFlyaways;
+										FilterGroups->TargetTriangleCount = FilterlSettings.NumTriangles;
+										FilterGroups->NumberOfTexturesInAtlas = FilterlSettings.NumTextures;
+										
+										// Advanced properties
+										FilterGroups->UseMultiCardClumps = FilterlSettings.AdvancedOptions.bUseMultiCardClumps;
+										FilterGroups->UseAdaptiveSubdivision = FilterlSettings.AdvancedOptions.bUseAdaptiveSubdivision;
+										FilterGroups->MaxVerticalSegmentsPerCard = FilterlSettings.AdvancedOptions.MaxVerticalSegmentsPerCard;
+										FilterGroups->StrandWidthScalingFactor = FilterlSettings.AdvancedOptions.StrandWidthScalingFactor;
+										FilterGroups->UseOptimizedCompressionFactor = FilterlSettings.AdvancedOptions.UseOptimizedCompressionFactor;
 									}
 								}
 							}
 						}
+						FHairCardGeneratorUtils::LoadGroomStrands(LocalAsset, [&GeometryCollection](TArray<TArray<FVector>>& CurvesPositions)
+						{
+							UE::CardGen::Private::BuildStrandsPositions(GeometryCollection, CurvesPositions);
+						});
+						OutputSettings.Last().GenerationSettings->UpdateStrandFilterAssignment();
 					}
-
-					FHairCardGeneratorUtils::LoadGroomStrands(LocalAsset, [&GeometryCollection](TArray<TArray<FVector>>& CurvesPositions)
-					{
-						UE::CardGen::Private::BuildStrandsPositions(GeometryCollection, CurvesPositions);
-					});
-					OutputSettings.Last().GenerationSettings->UpdateStrandFilterAssignment();
 				}
 			}
 		}
-		else
-		{
-			SetError(Context, &CardsSettings, TEXT("No valid groom cards groups. Please add a new group in the groom hair cards tab"));
-		}
 		SetValue(Context, MoveTemp(OutputSettings), &CardsSettings);
 	}
 }
diff --git a/Engine/Plugins/Experimental/HairCardGenerator/Source/HairCardGeneratorDataflow/Private/GenerateCardsClumpsNode.cpp b/Engine/Plugins/Experimental/HairCardGenerator/Source/HairCardGeneratorDataflow/Private/GenerateCardsClumpsNode.cpp
--- a/Engine/Plugins/Experimental/HairCardGenerator/Source/HairCardGeneratorDataflow/Private/GenerateCardsClumpsNode.cpp
+++ b/Engine/Plugins/Experimental/HairCardGenerator/Source/HairCardGeneratorDataflow/Private/GenerateCardsClumpsNode.cpp
@@ -77,16 +77,19 @@ void FGenerateCardsClumpsNode::Evaluate(UE::Dataflow::FContext& Context, const F
 						{
 							TArray<int32> FilterClumps; int32 ClumpCount = 0;
 							if(FHairCardGeneratorUtils::GenerateCardsClumps(GenerationSettings, FilterIndex, GenFlags, FilterClumps, ClumpCount))
-							{ 
+							{
+								int32 NumFilterClumps = 0;
 								for (int32 CurveIndex = 0, NumCurves = FilterClumps.Num()-1; CurveIndex < NumCurves; ++CurveIndex)
 								{
 									if (StrandsClumps.IsValidIndex(CurveIndex) && (FilterClumps[CurveIndex] != INDEX_NONE))
 									{
 										StrandsClumps[CurveIndex] = FilterClumps[CurveIndex] + NumClumps;
 										StrandsFilter[CurveIndex] = FilterIndex;
+										NumFilterClumps = FMath::Max(NumFilterClumps,  FilterClumps[CurveIndex]);
 									}
 								}
-								NumClumps += ClumpCount;
+								NumFilterClumps += 1;
+								NumClumps += NumFilterClumps;
 								return true;
 							}
 							return false;
diff --git a/Engine/Plugins/Experimental/HairCardGenerator/Source/HairCardGeneratorDataflow/Private/HairCardDataflowRendering.cpp b/Engine/Plugins/Experimental/HairCardGenerator/Source/HairCardGeneratorDataflow/Private/HairCardDataflowRendering.cpp
--- a/Engine/Plugins/Experimental/HairCardGenerator/Source/HairCardGeneratorDataflow/Private/HairCardDataflowRendering.cpp
+++ b/Engine/Plugins/Experimental/HairCardGenerator/Source/HairCardGeneratorDataflow/Private/HairCardDataflowRendering.cpp
@@ -6,8 +6,6 @@
 #include "GenerateCardsClumpsNode.h"
 #include "GenerateCardsTexturesNode.h"
 #include "PlanarCut.h"
-#include "Components/SphereComponent.h"
-#include "Components/TextRenderComponent.h"
 #include "Dataflow/DataflowRenderingViewMode.h"
 #include "DataflowRendering/DataflowRenderableComponents.h"
 #include "DataflowRendering/DataflowRenderableType.h"
@@ -125,28 +123,36 @@ namespace UE::CardGen::Private
 	{
 		if(!FilterTransforms.IsEmpty())
 		{
-			int32 FilterIndex = 0;
 			for(const TArray<int32>& TransformsIndices : FilterTransforms)
 			{
-				if (USphereComponent* FilterComponent = OutComponents.AddNewComponent<USphereComponent>())
+				if (!TransformsIndices.IsEmpty())
 				{
-					int32 LocalIndex = 0;
-					for(const int32& TransformIndex : TransformsIndices)
+					int32 TransformIndex = TransformsIndices[0];
+					if (UDynamicMeshComponent* FilterComponent = OutComponents.AddNewComponent<UDynamicMeshComponent>())
 					{
 						if (DynamicMeshes[TransformIndex].VertexCount() > 0)
 						{
-							if (UDynamicMeshComponent* Component = OutComponents.AddNewComponent<UDynamicMeshComponent>(NAME_None, FilterComponent))
+							FilterComponent->SetMesh(MoveTemp(DynamicMeshes[TransformIndex]));
+							FilterComponent->SetOverrideRenderMaterial(Dataflow::Rendering::GetVertexMaterial());
+							FilterComponent->SetCastShadow(false);
+						}
+
+						for (int32 LocalIndex = 1; LocalIndex < TransformsIndices.Num(); ++LocalIndex)
+						{
+							TransformIndex = TransformsIndices[LocalIndex];
+							if (DynamicMeshes[TransformIndex].VertexCount() > 0)
 							{
-								Component->SetupAttachment(FilterComponent);
-								Component->SetMesh(MoveTemp(DynamicMeshes[TransformIndex]));
-								Component->SetOverrideRenderMaterial(Dataflow::Rendering::GetVertexMaterial());
-								Component->SetCastShadow(false);
+								if (UDynamicMeshComponent* Component = OutComponents.AddNewComponent<UDynamicMeshComponent>(NAME_None, FilterComponent))
+								{
+									Component->SetupAttachment(FilterComponent);
+									Component->SetMesh(MoveTemp(DynamicMeshes[TransformIndex]));
+									Component->SetOverrideRenderMaterial(Dataflow::Rendering::GetVertexMaterial());
+									Component->SetCastShadow(false);
+								}
 							}
 						}
-						++LocalIndex;
 					}
 				}
-				++FilterIndex;
 			}
 		}
 		else
@@ -594,7 +600,6 @@ namespace UE::CardGen::Private
 		{
 			TArray<Geometry::FDynamicMesh3> DynamicMeshes;
 			TArray<TArray<int32>> FilterCardsIndices;
-			FilterCardsIndices.SetNum(2);
 			if (Instance.HasUptoDateCachedValue())
 			{
 				DynamicMeshes = Instance.GetCachedValue<TArray<Geometry::FDynamicMesh3>>();
@@ -620,19 +625,32 @@ namespace UE::CardGen::Private
 							FGenerateCardsTexturesNode::ObjectTextureIndicesAttribute, FName(CardsObjectsLODGroup));
 				
 						DynamicMeshes.SetNum(NumCards);
+						
+						TArray<int32> TextureIndices;
+						TextureIndices.Init(INDEX_NONE, NumCards);
+						FilterCardsIndices.SetNum(NumCards);
+						
+						int32 NumTextures = 0;
+						for(int32 CardIndex = 0; CardIndex < NumCards; ++CardIndex)
+						{
+							if (CardsTextureArray.IsValidIndex(CardIndex) && CardsTextureArray[CardIndex] == CardIndex)
+							{
+								TextureIndices[CardIndex] = NumTextures;
+								FilterCardsIndices[NumTextures].Add(CardIndex);
+								++NumTextures;
+							}
+						}
+						FilterCardsIndices.SetNum(NumTextures);
+						
 						for(int32 CardIndex = 0; CardIndex < NumCards; ++CardIndex)
 						{
 							const int32 TextureIndex = (CardIndex < CardsTextureArray.Num()) ? CardsTextureArray[CardIndex] : INDEX_NONE;
 							FLinearColor CardColor = (TextureIndex != INDEX_NONE) ? FLinearColor::IntToDistinctColor(TextureIndex, 0.75f, 1.0f, 90.0f) : FLinearColor::Black;
 							CardColor *= (TextureIndex == CardIndex) ? 1.0 : GGroomDataflowCardsAlpha;
 
-							if(TextureIndex == CardIndex)
-							{
-								FilterCardsIndices[0].Add(CardIndex);
-							}
-							else
+							if(TextureIndex != CardIndex && TextureIndices.IsValidIndex(TextureIndex) && FilterCardsIndices.IsValidIndex(TextureIndices[TextureIndex]))
 							{
-								FilterCardsIndices[1].Add(CardIndex);
+								FilterCardsIndices[TextureIndices[TextureIndex]].Add(CardIndex);
 							}
 
 							TArray<FLinearColor> VertexColors;
diff --git a/Engine/Plugins/Experimental/HairCardGenerator/Source/HairCardGeneratorDataflow/Public/BuildCardsSettingsNode.h b/Engine/Plugins/Experimental/HairCardGenerator/Source/HairCardGeneratorDataflow/Public/BuildCardsSettingsNode.h
--- a/Engine/Plugins/Experimental/HairCardGenerator/Source/HairCardGeneratorDataflow/Public/BuildCardsSettingsNode.h
+++ b/Engine/Plugins/Experimental/HairCardGenerator/Source/HairCardGeneratorDataflow/Public/BuildCardsSettingsNode.h
@@ -33,6 +33,34 @@ struct FGroomCardsSettings
     TObjectPtr<UGroomAsset> GroomAsset = nullptr;
 };
 
+
+USTRUCT()
+struct FGroomAdvancedFilterSettings
+{
+	GENERATED_BODY()
+	
+	// Generate multiple cards (3) per strand clump to give the appearance of volume
+	UPROPERTY(EditAnywhere, AdvancedDisplay, Category="Cards Filter")
+	bool bUseMultiCardClumps = true;
+
+	// Use adaptive subdivision when generating card geometry
+	UPROPERTY(EditAnywhere, AdvancedDisplay, Category="Cards Filter")
+	bool bUseAdaptiveSubdivision = true;
+
+	// Maximum number of segments along each card (aligned with the strands), ignored for adaptive subdivision
+	UPROPERTY(EditAnywhere, AdvancedDisplay, Category="Cards Filter")
+	int32 MaxVerticalSegmentsPerCard = 10;
+
+	// Scaling factor mapping strand-width to pixel-size for coverage texture
+	UPROPERTY(EditAnywhere, AdvancedDisplay, Category="Cards Filter")
+	float StrandWidthScalingFactor = 1.0f;
+
+	// Compress textures along strand direction to save vertical redolution,
+	// if strands are all moving in nearly the same direction
+	UPROPERTY(EditAnywhere, AdvancedDisplay, Category="Cards Filter")
+	bool UseOptimizedCompressionFactor = true;
+};
+
 USTRUCT()
 struct FGroomFilterSettings
 {
@@ -41,14 +69,6 @@ struct FGroomFilterSettings
 	// Filter group name to be identified in the override settings
 	UPROPERTY(EditAnywhere, Category="Groom Cards")
 	FName FilterName;
-	
-	// LOD index of the filter group settings
-	UPROPERTY(EditAnywhere, Category="Groom Cards")
-	int32 LODIndex = -1;
-
-	// Group index of the filter group settings
-	UPROPERTY(EditAnywhere, Category="Groom Cards")
-	int32 GroupIndex = 0;
 
 	// Total number of cards to generate for this LOD settings
 	UPROPERTY(EditAnywhere, Category="Groom Cards", meta=(ClampMin="1", ClampMax="100000", DisplayName = "Num Clumps"))
@@ -69,6 +89,66 @@ struct FGroomFilterSettings
 	// Card group names that will belong to the filter group settings
 	UPROPERTY(EditAnywhere, Category="Groom Cards")
 	TArray<FName> CardGroups;
+
+	// Advanced options for the filter group settings
+	UPROPERTY(EditAnywhere, Category="Groom Cards")
+	FGroomAdvancedFilterSettings AdvancedOptions;
+};
+
+USTRUCT()
+struct FGroomAdvancedGenerationSettings
+{
+	GENERATED_BODY()
+	
+	// Use previous LOD generated cards and textures but reduce triangle count and flyaways
+	UPROPERTY(EditAnywhere, AdvancedDisplay,Category="Cards Generation")
+	bool bReduceCardsFromPreviousLOD = false;
+
+	// Generate geometry for all groom groups on group 0
+	UPROPERTY(EditAnywhere, AdvancedDisplay, Category="Cards Generation", meta=(AllowPrivateAccess))
+	bool bGenerateGeometryForAllGroups = true;
+
+	// Seed value for pseudo-random number generation (set to a specific value for repeatable results)
+	UPROPERTY(EditAnywhere, AdvancedDisplay,Category="Cards Generation", meta=(ClampMin="0", ClampMax="10000"))
+	int32 RandomSeed = 0;
+
+	// Place new card textures in reserved space from previous LOD
+	UPROPERTY(EditAnywhere, AdvancedDisplay, Category="Texture Rendering")
+	bool bUseReservedSpaceFromPreviousLOD = false;
+
+	// Size of hair card texture atlases
+	UPROPERTY(EditAnywhere, AdvancedDisplay, Category="Texture Rendering")
+	EHairCardAtlasSize AtlasSize = EHairCardAtlasSize::AtlasSize4096;
+
+	// Percentage of texture atlas space to reserve for higher LODs 
+	UPROPERTY(EditAnywhere, AdvancedDisplay, Category="Texture Rendering", meta=(ClampMin="0", ClampMax="75"))
+	int32 ReserveTextureSpaceLOD = 0;
+
+	// Use strand width
+	UPROPERTY(EditAnywhere, AdvancedDisplay, Category="Texture Rendering")
+	bool bUseGroomAssetStrandWidth = true;
+};
+
+USTRUCT()
+struct FGroomGenerationSettings
+{
+	GENERATED_BODY()
+	
+	// LOD index of the filter group settings
+	UPROPERTY(EditAnywhere, Category="Groom Cards")
+	int32 LODIndex = 0;
+
+	// Group index of the filter group settings
+	UPROPERTY(EditAnywhere, Category="Groom Cards")
+	int32 GroupIndex = 0;
+
+	// Lisdt of filter settings to use for that generation
+	UPROPERTY(EditAnywhere, Category="Groom Cards")
+	TArray<FGroomFilterSettings> FilterSettings;
+
+	// Advanced options for the generation settings
+	UPROPERTY(EditAnywhere, AdvancedDisplay,Category="Groom Cards")
+	FGroomAdvancedGenerationSettings AdvancedOptions;
 };
 
 /** Build the cards generation settings */
@@ -108,6 +188,6 @@ private:
 	
 	/** List of filter settings to override */
 	UPROPERTY(EditAnywhere, Category = "Groom Cards")
-	TArray<FGroomFilterSettings> FilterSettings;
+	TArray<FGroomGenerationSettings> GenerationSettings;
 };
 
diff --git a/Engine/Plugins/Experimental/HairCardGenerator/Source/HairCardGeneratorEditor/Private/HairCardGeneratorEditorModule.cpp b/Engine/Plugins/Experimental/HairCardGenerator/Source/HairCardGeneratorEditor/Private/HairCardGeneratorEditorModule.cpp
--- a/Engine/Plugins/Experimental/HairCardGenerator/Source/HairCardGeneratorEditor/Private/HairCardGeneratorEditorModule.cpp
+++ b/Engine/Plugins/Experimental/HairCardGenerator/Source/HairCardGeneratorEditor/Private/HairCardGeneratorEditorModule.cpp
@@ -121,7 +121,7 @@ static FAutoConsoleCommand CCmdLoadGroom(
  *****************************************************************************/
 
 bool FHairCardGeneratorUtils::BuildGenerationSettings(const bool bQuerySettings, UGroomAsset* NewGroomAsset, FHairGroupsCardsSourceDescription& CardsDesc,
-	TObjectPtr<UHairCardGeneratorPluginSettings>& GenerationSettings, uint8& AllGenFlags, uint8& PipelineGenFlags)
+	TObjectPtr<UHairCardGeneratorPluginSettings>& GenerationSettings, uint8& AllGenFlags, uint8& PipelineGenFlags, const FCardsGenerationAdvancedOptions& AdvancedOptions)
 {
 	// Find the HairCardGenController in Python (or wherever it may be implemented -- could be moved to C++/Blueprints/etc.)
 	TObjectPtr<UHairCardGenControllerBase> CardGenController = FHairCardGenController::GetHairCardGenerator();
@@ -152,6 +152,16 @@ bool FHairCardGeneratorUtils::BuildGenerationSettings(const bool bQuerySettings,
 			return false;
 		}
 	}
+	else
+	{
+		GenerationSettings->bGenerateGeometryForAllGroups = AdvancedOptions.bGenerateGeometryForAllGroups;
+		GenerationSettings->bReduceCardsFromPreviousLOD = AdvancedOptions.bReduceCardsFromPreviousLOD;
+		GenerationSettings->RandomSeed = AdvancedOptions.RandomSeed;
+		GenerationSettings->bUseReservedSpaceFromPreviousLOD = AdvancedOptions.bUseReservedSpaceFromPreviousLOD;
+		GenerationSettings->AtlasSize = static_cast<EHairCardAtlasSize>(AdvancedOptions.AtlasSize);
+		GenerationSettings->ReserveTextureSpaceLOD = AdvancedOptions.ReserveTextureSpaceLOD;
+		GenerationSettings->bUseGroomAssetStrandWidth = AdvancedOptions.bUseGroomAssetStrandWidth;
+	}
 
 	// Check primary settings object merged and all group settings for differences
 	AllGenFlags = GenerationSettings->GetAllPipelineGeneratedDifferences();
@@ -942,7 +952,7 @@ bool FHairCardGeneratorEditorModule::GenerateHairCardsForLOD(UGroomAsset* NewGro
 {
 	TObjectPtr<UHairCardGeneratorPluginSettings> GenerationSettings = nullptr;
 	uint8 AllGenFlags = 0, PipelineGenFlags = 0;
-	if(!FHairCardGeneratorUtils::BuildGenerationSettings(true, NewGroomAsset, CardsDesc, GenerationSettings, AllGenFlags, PipelineGenFlags))
+	if(!FHairCardGeneratorUtils::BuildGenerationSettings(true, NewGroomAsset, CardsDesc, GenerationSettings, AllGenFlags, PipelineGenFlags, FCardsGenerationAdvancedOptions()))
 	{
 		return false;
 	}
diff --git a/Engine/Plugins/Experimental/HairCardGenerator/Source/HairCardGeneratorEditor/Public/HairCardGeneratorEditorModule.h b/Engine/Plugins/Experimental/HairCardGenerator/Source/HairCardGeneratorEditor/Public/HairCardGeneratorEditorModule.h
--- a/Engine/Plugins/Experimental/HairCardGenerator/Source/HairCardGeneratorEditor/Public/HairCardGeneratorEditorModule.h
+++ b/Engine/Plugins/Experimental/HairCardGenerator/Source/HairCardGeneratorEditor/Public/HairCardGeneratorEditorModule.h
@@ -9,6 +9,30 @@ class UGroomAsset;
 class UHairCardGeneratorPluginSettings;
 class UHairCardGenControllerBase;
 
+struct HAIRCARDGENERATOREDITOR_API FCardsGenerationAdvancedOptions
+{
+	// Use previous LOD generated cards and textures but reduce triangle count and flyaways
+	bool bReduceCardsFromPreviousLOD = false;
+
+	// Generate geometry for all groom groups on group 0
+	bool bGenerateGeometryForAllGroups = true;
+
+	// Seed value for pseudo-random number generation (set to a specific value for repeatable results)
+	int32 RandomSeed = 0;
+
+	// Place new card textures in reserved space from previous LOD
+	bool bUseReservedSpaceFromPreviousLOD = false;
+
+	// Size of hair card texture atlases
+	uint8 AtlasSize = 12;
+
+	// Percentage of texture atlas space to reserve for higher LODs 
+	int32 ReserveTextureSpaceLOD = 0;
+
+	// Use strand width
+	bool bUseGroomAssetStrandWidth = true;
+};
+
 /** FHairCardGeneratorUtils has static interface functions to be used in other modules */
 struct HAIRCARDGENERATOREDITOR_API FHairCardGeneratorUtils
 {
@@ -17,7 +41,7 @@ struct HAIRCARDGENERATOREDITOR_API FHairCardGeneratorUtils
 	
 	/** Build the generation settings for a given grom asset and a cards description */
 	static bool BuildGenerationSettings(const bool bQuerySettings, UGroomAsset* GroomAsset, FHairGroupsCardsSourceDescription& CardsDesc,
-		TObjectPtr<UHairCardGeneratorPluginSettings>& GenerationSettings, uint8& GenerationFlags, uint8& PipelineFlags);
+		TObjectPtr<UHairCardGeneratorPluginSettings>& GenerationSettings, uint8& GenerationFlags, uint8& PipelineFlags, const FCardsGenerationAdvancedOptions& AdvancedOptions);
 
 	/** Load the generation settings for a given grom asset and a cards description */
 	static bool LoadGenerationSettings(const TObjectPtr<UHairCardGeneratorPluginSettings>& GenerationSettings);
