Fix dataflow cards rendering crash when the generate LOD from previous is on

diff --git a/Engine/Plugins/Experimental/HairCardGenerator/Source/HairCardGeneratorDataflow/Private/GenerateCardsTexturesNode.cpp b/Engine/Plugins/Experimental/HairCardGenerator/Source/HairCardGeneratorDataflow/Private/GenerateCardsTexturesNode.cpp
--- a/Engine/Plugins/Experimental/HairCardGenerator/Source/HairCardGeneratorDataflow/Private/GenerateCardsTexturesNode.cpp
+++ b/Engine/Plugins/Experimental/HairCardGenerator/Source/HairCardGeneratorDataflow/Private/GenerateCardsTexturesNode.cpp
@@ -70,15 +70,18 @@ void FGenerateCardsTexturesNode::Evaluate(UE::Dataflow::FContext& Context, const
 							int32 TexturesCount = 0;
 							if(FHairCardGeneratorUtils::GenerateCardsTexturesClusters(GenerationSettings, FilterIndex, GenFlags, FilterTextures, TexturesCount))
 							{
-								const int32 CardCount = FilterTextures.Num() - 1 - TexturesCount;
-								const int32 CardOffset = CardsTextures.Num();
-
-								CardsTextures.Reserve(CardOffset + CardCount);
-								for (int32 CardIndex = 0; CardIndex < CardCount; ++CardIndex)
+								if (!FilterTextures.IsEmpty())
 								{
-									CardsTextures.Add((FilterTextures[CardIndex] < TexturesCount) ? FilterTextures[CardCount + FilterTextures[CardIndex]] + CardOffset : INDEX_NONE);
+									const int32 CardCount = FilterTextures.Num() - 1 - TexturesCount;
+									const int32 CardOffset = CardsTextures.Num();
+
+									CardsTextures.Reserve(CardOffset + CardCount);
+									for (int32 CardIndex = 0; CardIndex < CardCount; ++CardIndex)
+									{
+										CardsTextures.Add((FilterTextures[CardIndex] < TexturesCount) ? FilterTextures[CardCount + FilterTextures[CardIndex]] + CardOffset : INDEX_NONE);
+									}
+									return true;
 								}
-								return true;
 							}
 							return false;
 						}, false);
@@ -88,42 +91,45 @@ void FGenerateCardsTexturesNode::Evaluate(UE::Dataflow::FContext& Context, const
 							TArray<float> FilterUVs;
 							if(FHairCardGeneratorUtils::GenerateTexturesLayoutAndAtlases(GenerationSettings, LODSettings.PipelineFlags, FilterUVs))
 							{
-								VertexUVs.Reserve(FilterUVs.Num());
-								int32 UVsIndex = 0;
-								while (UVsIndex < FilterUVs.Num())
+								if (!FilterUVs.IsEmpty())
 								{
-									if (FilterUVs[UVsIndex] != -1.0)
+									VertexUVs.Reserve(FilterUVs.Num());
+									int32 UVsIndex = 0;
+									while (UVsIndex < FilterUVs.Num())
 									{
-										VertexUVs.Add(FVector2f(FilterUVs[UVsIndex], FilterUVs[UVsIndex + 1]));
-										UVsIndex += 2;
+										if ((FilterUVs[UVsIndex] != -1.0) && ((UVsIndex + 1) < FilterUVs.Num()))
+										{
+											VertexUVs.Add(FVector2f(FilterUVs[UVsIndex], FilterUVs[UVsIndex + 1]));
+											UVsIndex += 2;
+										}
+										else
+										{
+											UVsIndex += 1;
+										}
 									}
-									else
+									FString CardsObjectsLODGroup = CardsObjectsGroup.ToString();
+									CardsObjectsLODGroup.AppendInt(GenerationSettings->GetLODIndex());
+
+									FString CardsVerticesLODGroup = FGenerateCardsGeometryNode::CardsVerticesGroup.ToString();
+									CardsVerticesLODGroup.AppendInt(GenerationSettings->GetLODIndex());
+
+									TManagedArray<int32>& ObjectTextureIndices = GroomCollection.AddAttribute<int32>(ObjectTextureIndicesAttribute, FName(CardsObjectsLODGroup));
+									TManagedArray<FVector2f>& VertexTextureUvs = GroomCollection.AddAttribute<FVector2f>(VertexTextureUVsAttribute, FName(CardsVerticesLODGroup));
+
+									GroomCollection.EmptyGroup(FName(CardsObjectsLODGroup));
+									GroomCollection.AddElements(CardsTextures.Num(), FName(CardsObjectsLODGroup));
+
+									for (int32 CardIndex = 0, NumCards = CardsTextures.Num(); CardIndex < NumCards; ++CardIndex)
 									{
-										UVsIndex += 1;
+										ObjectTextureIndices[CardIndex] = CardsTextures[CardIndex];
 									}
-								}
-								FString CardsObjectsLODGroup = CardsObjectsGroup.ToString();
-								CardsObjectsLODGroup.AppendInt(GenerationSettings->GetLODIndex());
-			
-								FString CardsVerticesLODGroup = FGenerateCardsGeometryNode::CardsVerticesGroup.ToString();
-								CardsVerticesLODGroup.AppendInt(GenerationSettings->GetLODIndex());
-							
-								TManagedArray<int32>& ObjectTextureIndices = GroomCollection.AddAttribute<int32>(ObjectTextureIndicesAttribute, FName(CardsObjectsLODGroup));
-								TManagedArray<FVector2f>& VertexTextureUvs = GroomCollection.AddAttribute<FVector2f>(VertexTextureUVsAttribute, FName(CardsVerticesLODGroup));
-							
-								GroomCollection.EmptyGroup(FName(CardsObjectsLODGroup));
-								GroomCollection.AddElements(CardsTextures.Num(), FName(CardsObjectsLODGroup));
-							
-								for(int32 CardIndex = 0, NumCards = CardsTextures.Num(); CardIndex < NumCards; ++CardIndex)
-								{
-									ObjectTextureIndices[CardIndex] = CardsTextures[CardIndex];
-								}
 
-								if(VertexUVs.Num() == VertexTextureUvs.Num())
-								{
-									for(int32 VertexIndex = 0, NumVertices = VertexUVs.Num(); VertexIndex < NumVertices; ++VertexIndex)
+									if (VertexUVs.Num() == VertexTextureUvs.Num())
 									{
-										VertexTextureUvs[VertexIndex] = VertexUVs[VertexIndex];
+										for (int32 VertexIndex = 0, NumVertices = VertexUVs.Num(); VertexIndex < NumVertices; ++VertexIndex)
+										{
+											VertexTextureUvs[VertexIndex] = VertexUVs[VertexIndex];
+										}
 									}
 								}
 							}
