HairCardGenerator: Add functionality to pull AO strand vert properties into hair card generator python groom definition.

diff --git a/Engine/Plugins/Experimental/HairCardGenerator/Source/HairCardGeneratorEditor/Private/HairCardGenControllerBase.h b/Engine/Plugins/Experimental/HairCardGenerator/Source/HairCardGeneratorEditor/Private/HairCardGenControllerBase.h
--- a/Engine/Plugins/Experimental/HairCardGenerator/Source/HairCardGeneratorEditor/Private/HairCardGenControllerBase.h
+++ b/Engine/Plugins/Experimental/HairCardGenerator/Source/HairCardGeneratorEditor/Private/HairCardGenControllerBase.h
@@ -45,6 +45,9 @@ struct FHairCardGen_GroomData
 	UPROPERTY(VisibleAnywhere, BlueprintReadOnly, Category = "Groom")
 	TArray<float> VertexWidths;
 
+	UPROPERTY(VisibleAnywhere, BlueprintReadOnly, Category = "Groom")
+	TArray<float> VertexOcclusions;
+
 	UPROPERTY(VisibleAnywhere, BlueprintReadOnly, Category = "Groom")
 	int32 HairlineGroupID = INDEX_NONE;
 };
diff --git a/Engine/Plugins/Experimental/HairCardGenerator/Source/HairCardGeneratorEditor/Private/HairCardGenStrandFilter.cpp b/Engine/Plugins/Experimental/HairCardGenerator/Source/HairCardGeneratorEditor/Private/HairCardGenStrandFilter.cpp
--- a/Engine/Plugins/Experimental/HairCardGenerator/Source/HairCardGeneratorEditor/Private/HairCardGenStrandFilter.cpp
+++ b/Engine/Plugins/Experimental/HairCardGenerator/Source/HairCardGeneratorEditor/Private/HairCardGenStrandFilter.cpp
@@ -120,6 +120,25 @@ TArrayView<const float> FHairStrandAttribsRefProxy::GetVertexWidths() const
     return TArrayView<const float>(VertexWidths.GetArrayView(StartVertIndex).GetData(), VertCount);
 }
 
+bool FHairStrandAttribsRefProxy::HasVertexAO() const
+{
+	return VertexAOs.IsValid();
+}
+
+const float& FHairStrandAttribsRefProxy::GetVertexAO(int Index) const
+{
+	const FVertexID VertID(StartVertIndex + Index);
+	return VertexAOs[VertID];
+}
+
+TArrayView<const float> FHairStrandAttribsRefProxy::GetVertexAOs() const
+{
+	const FStrandID StrandID(StrandIndex);
+	int VertCount = StrandVertCounts[StrandID];
+
+	return TArrayView<const float>(VertexAOs.GetArrayView(StartVertIndex).GetData(), VertCount);
+}
+
 
 FHairStrandAttribsRefProxy::FHairStrandAttribsRefProxy(const FHairDescription& InHairDescription)
     : StrandIndex(0)
@@ -134,6 +153,7 @@ FHairStrandAttribsRefProxy::FHairStrandAttribsRefProxy(const FHairDescription& I
     , StrandVertCounts(InHairDescription.StrandAttributes().GetAttributesRef<int>(HairAttribute::Strand::VertexCount))
     , VertexPositions(InHairDescription.VertexAttributes().GetAttributesRef<FVector3f>(HairAttribute::Vertex::Position))
     , VertexWidths(InHairDescription.VertexAttributes().GetAttributesRef<float>(HairAttribute::Vertex::Width))
+	, VertexAOs(InHairDescription.VertexAttributes().GetAttributesRef<float>(HairAttribute::Vertex::AO))
 {}
 
 
diff --git a/Engine/Plugins/Experimental/HairCardGenerator/Source/HairCardGeneratorEditor/Private/HairCardGenStrandFilter.h b/Engine/Plugins/Experimental/HairCardGenerator/Source/HairCardGeneratorEditor/Private/HairCardGenStrandFilter.h
--- a/Engine/Plugins/Experimental/HairCardGenerator/Source/HairCardGeneratorEditor/Private/HairCardGenStrandFilter.h
+++ b/Engine/Plugins/Experimental/HairCardGenerator/Source/HairCardGeneratorEditor/Private/HairCardGenStrandFilter.h
@@ -48,6 +48,10 @@ public:
     const float& GetVertexWidth(int Index) const;
     TArrayView<const float> GetVertexWidths() const;
 
+	bool HasVertexAO() const;
+	const float& GetVertexAO(int Index) const;
+	TArrayView<const float> GetVertexAOs() const;
+
 
 private:
     FHairStrandAttribsRefProxy() = delete;
@@ -70,6 +74,7 @@ private:
     TStrandAttributesConstRef<int> StrandVertCounts;
     TVertexAttributesConstRef<FVector3f> VertexPositions;
     TVertexAttributesConstRef<float> VertexWidths;
+	TVertexAttributesConstRef<float> VertexAOs;
 };
 
 
diff --git a/Engine/Plugins/Experimental/HairCardGenerator/Source/HairCardGeneratorEditor/Private/HairCardGenStrandNCSInterpolator.cpp b/Engine/Plugins/Experimental/HairCardGenerator/Source/HairCardGeneratorEditor/Private/HairCardGenStrandNCSInterpolator.cpp
--- a/Engine/Plugins/Experimental/HairCardGenerator/Source/HairCardGeneratorEditor/Private/HairCardGenStrandNCSInterpolator.cpp
+++ b/Engine/Plugins/Experimental/HairCardGenerator/Source/HairCardGeneratorEditor/Private/HairCardGenStrandNCSInterpolator.cpp
@@ -3,7 +3,7 @@
 #include "HairCardGenStrandNCSInterpolator.h"
 #include <iostream>
 
-FHairCardGenStrandNSCInterpolator::FHairCardGenStrandNSCInterpolator(const TArray<FVector>& Positions, const TArray<float>& Widths)
+FHairCardGenStrandNSCInterpolator::FHairCardGenStrandNSCInterpolator(const TArray<FVector>& Positions, const TArray<float>& Widths, const TArray<float>& AOs)
 {
     int NumPoints = Positions.Num();
 
@@ -30,14 +30,16 @@ FHairCardGenStrandNSCInterpolator::FHairCardGenStrandNSCInterpolator(const TArra
 		DistancesFromStart[i] = (DistancesFromStart[i] + float(i) * 0.1 * DistancesFromStart.Last()) / ((0.1 * float(NumPoints) + 0.9) * DistancesFromStart.Last());
 	}
 
-	Splines.Reserve(4);
+	Splines.Reserve(5);
 
 	Splines.Add(FHairCardGenNaturalCubicSplines(DistancesFromStart, X));
 	Splines.Add(FHairCardGenNaturalCubicSplines(DistancesFromStart, Y));
 	Splines.Add(FHairCardGenNaturalCubicSplines(DistancesFromStart, Z));
 	Splines.Add(FHairCardGenNaturalCubicSplines(DistancesFromStart, Widths));
-
-	Splines.Shrink();
+	if (AOs.Num() > 0)
+	{
+		Splines.Add(FHairCardGenNaturalCubicSplines(DistancesFromStart, AOs));
+	}
 }
 
 FHairCardGenStrandNSCInterpolator::StrandInterpolationResult FHairCardGenStrandNSCInterpolator::GetInterpolatedStrand(const int NumInterpolatedPoints)
@@ -46,6 +48,10 @@ FHairCardGenStrandNSCInterpolator::StrandInterpolationResult FHairCardGenStrandN
 
 	Result.Positions.SetNum(NumInterpolatedPoints * 3);
 	Result.Widths.SetNum(NumInterpolatedPoints);
+	if (Splines.Num() == 5)
+	{
+		Result.AOs.SetNum(NumInterpolatedPoints);
+	}
 
 	for (int i = 0; i < NumInterpolatedPoints; i++)
 	{
@@ -56,6 +62,11 @@ FHairCardGenStrandNSCInterpolator::StrandInterpolationResult FHairCardGenStrandN
 			Result.Positions[i * 3 + j] = Splines[j](T);
 		}
 		Result.Widths[i] = Splines[3](T);
+
+		if (Splines.Num() == 5)
+		{
+			Result.AOs[i] = Splines[4](T);
+		}
 	}
 
 	return Result;
diff --git a/Engine/Plugins/Experimental/HairCardGenerator/Source/HairCardGeneratorEditor/Private/HairCardGenStrandNCSInterpolator.h b/Engine/Plugins/Experimental/HairCardGenerator/Source/HairCardGeneratorEditor/Private/HairCardGenStrandNCSInterpolator.h
--- a/Engine/Plugins/Experimental/HairCardGenerator/Source/HairCardGeneratorEditor/Private/HairCardGenStrandNCSInterpolator.h
+++ b/Engine/Plugins/Experimental/HairCardGenerator/Source/HairCardGeneratorEditor/Private/HairCardGenStrandNCSInterpolator.h
@@ -16,9 +16,10 @@ public:
     {
         TArray<float> Positions;
         TArray<float> Widths;
+		TArray<float> AOs;
     };
 
 public:
-	FHairCardGenStrandNSCInterpolator(const TArray<FVector>& Positions, const TArray<float>& Widths);
+	FHairCardGenStrandNSCInterpolator(const TArray<FVector>& Positions, const TArray<float>& Widths, const TArray<float>& AOs);
     StrandInterpolationResult GetInterpolatedStrand(const int NumInterpolatedPoints);
 }; 
diff --git a/Engine/Plugins/Experimental/HairCardGenerator/Source/HairCardGeneratorEditor/Private/HairCardGeneratorEditorModule.cpp b/Engine/Plugins/Experimental/HairCardGenerator/Source/HairCardGeneratorEditor/Private/HairCardGeneratorEditorModule.cpp
--- a/Engine/Plugins/Experimental/HairCardGenerator/Source/HairCardGeneratorEditor/Private/HairCardGeneratorEditorModule.cpp
+++ b/Engine/Plugins/Experimental/HairCardGenerator/Source/HairCardGeneratorEditor/Private/HairCardGeneratorEditorModule.cpp
@@ -369,6 +369,11 @@ static bool GenerateHairCardGenGroomData(const int PointsPerStrand, const FHairD
 
 	Out.VertexPositions.Reserve(NumStrands * PointsPerStrand * 3);
 	Out.VertexWidths.Reserve(NumStrands * PointsPerStrand);
+	bool bHasAO = FirstStrand.HasVertexAO();
+	if (bHasAO)
+	{
+		Out.VertexOcclusions.Reserve(NumStrands * PointsPerStrand);
+	}
 
 	TArray<TArray<FVector>> StrandsPositions;
 	BuilderFunction(StrandsPositions);
@@ -385,24 +390,32 @@ static bool GenerateHairCardGenGroomData(const int PointsPerStrand, const FHairD
 		VertexPositions.SetNum(StrandAttribs.NumVerts());
 		TArray<float> VertexWidths;
 		VertexWidths.SetNum(StrandAttribs.NumVerts());
+		TArray<float> VertexAOs;
+		VertexAOs.Reserve(StrandAttribs.NumVerts());
 
 		const bool bIsVertexValid = StrandsPositions.IsValidIndex(StrandIndex) && (StrandAttribs.NumVerts() <= StrandsPositions[StrandIndex].Num());
 		for (int StrandVertIndex = 0; StrandVertIndex < StrandAttribs.NumVerts(); ++StrandVertIndex)
 		{
 			VertexPositions[StrandVertIndex] = bIsVertexValid ? StrandsPositions[StrandIndex][StrandVertIndex] : FVector(StrandAttribs.GetVertexPosition(StrandVertIndex));
 			VertexWidths[StrandVertIndex] = StrandAttribs.HasVertexWidths() ? StrandAttribs.GetVertexWidth(StrandVertIndex) : Strand_DefaultWidth;
+			if (bHasAO)
+			{
+				VertexAOs.Add(StrandAttribs.GetVertexAO(StrandVertIndex));
+			}
 		}
 
-		FHairCardGenStrandNSCInterpolator::StrandInterpolationResult InterpolatedStrand = FHairCardGenStrandNSCInterpolator(VertexPositions, VertexWidths).GetInterpolatedStrand(PointsPerStrand);
+		FHairCardGenStrandNSCInterpolator::StrandInterpolationResult InterpolatedStrand = FHairCardGenStrandNSCInterpolator(VertexPositions, VertexWidths, VertexAOs).GetInterpolatedStrand(PointsPerStrand);
 
 		Out.VertexPositions.Append(InterpolatedStrand.Positions);
 		Out.VertexWidths.Append(InterpolatedStrand.Widths);
+		Out.VertexOcclusions.Append(InterpolatedStrand.AOs);
 		++StrandIndex;
 	}
 
 	Out.Strands.Shrink();
 	Out.VertexPositions.Shrink();
 	Out.VertexWidths.Shrink();
+	Out.VertexOcclusions.Shrink();
 
 	Out.HairlineGroupID = INDEX_NONE;
 
