Groom dataflow : adjust default settings

#Synced-CL 51362814
#Synced-CL 51362814

diff --git a/Engine/Plugins/Experimental/HairCardGenerator/Source/HairCardGeneratorDataflow/Public/BuildCardsSettingsNode.h b/Engine/Plugins/Experimental/HairCardGenerator/Source/HairCardGeneratorDataflow/Public/BuildCardsSettingsNode.h
--- a/Engine/Plugins/Experimental/HairCardGenerator/Source/HairCardGeneratorDataflow/Public/BuildCardsSettingsNode.h
+++ b/Engine/Plugins/Experimental/HairCardGenerator/Source/HairCardGeneratorDataflow/Public/BuildCardsSettingsNode.h
@@ -72,19 +72,19 @@ struct FGroomFilterSettings
 
 	// Total number of cards to generate for this LOD settings
 	UPROPERTY(EditAnywhere, Category="Groom Cards", meta=(ClampMin="1", ClampMax="100000", DisplayName = "Num Clumps"))
-	int32 NumClumps = 50;
+	int32 NumClumps = 800;
 
 	// Maximum number of cards to assign to flyaway strands
 	UPROPERTY(EditAnywhere, Category="Groom Cards", meta=(ClampMin="0", ClampMax="1000"))
-	int32 NumFlyaways = 10;
+	int32 NumFlyaways = 100;
 
 	// Total number of triangles to generate for this LOD settings
 	UPROPERTY(EditAnywhere, Category="Groom Cards", meta=(ClampMin="1", ClampMax="100000"))
-	int32 NumTriangles = 2000;
+	int32 NumTriangles = 45000;
 
 	// Total number of textures to generate for this LOD settings
 	UPROPERTY(EditAnywhere, Category="Groom Cards")
-	int32 NumTextures = 75;
+	int32 NumTextures = 100;
 
 	// Card group names that will belong to the filter group settings
 	UPROPERTY(EditAnywhere, Category="Groom Cards")
diff --git a/Engine/Plugins/Experimental/HairCardGenerator/Source/HairCardGeneratorDataflow/Public/GenerateCardsClumpsNode.h b/Engine/Plugins/Experimental/HairCardGenerator/Source/HairCardGeneratorDataflow/Public/GenerateCardsClumpsNode.h
--- a/Engine/Plugins/Experimental/HairCardGenerator/Source/HairCardGeneratorDataflow/Public/GenerateCardsClumpsNode.h
+++ b/Engine/Plugins/Experimental/HairCardGenerator/Source/HairCardGeneratorDataflow/Public/GenerateCardsClumpsNode.h
@@ -22,11 +22,11 @@ struct FCardsClumpsSettings
 
 	// Total number of cards to generate for this LOD settings
 	UPROPERTY(EditAnywhere, Category="Groom Cards", meta=(ClampMin="1", ClampMax="100000", DisplayName = "Num Clumps"))
-	int32 NumCards = 50;
+	int32 NumCards = 800;
 
 	// Maximum number of cards to assign to flyaway strands
 	UPROPERTY(EditAnywhere, Category="Groom Cards", meta=(ClampMin="0", ClampMax="1000"))
-	int32 NumFlyaways = 10;
+	int32 NumFlyaways = 100;
 };
 
 /** Generate the clumps used to build the cards from the strands */
diff --git a/Engine/Plugins/Experimental/HairCardGenerator/Source/HairCardGeneratorDataflow/Public/GenerateCardsGeometryNode.h b/Engine/Plugins/Experimental/HairCardGenerator/Source/HairCardGeneratorDataflow/Public/GenerateCardsGeometryNode.h
--- a/Engine/Plugins/Experimental/HairCardGenerator/Source/HairCardGeneratorDataflow/Public/GenerateCardsGeometryNode.h
+++ b/Engine/Plugins/Experimental/HairCardGenerator/Source/HairCardGeneratorDataflow/Public/GenerateCardsGeometryNode.h
@@ -22,7 +22,7 @@ struct FCardsGeometrySettings
 
 	// Total number of triangles to generate for this LOD settings
 	UPROPERTY(EditAnywhere, Category="Groom Cards", meta=(ClampMin="1", ClampMax="100000"))
-	int32 NumTriangles = 2000;
+	int32 NumTriangles = 45000;
 };
 
 /** Generate the geometry used to build the cards from the strands */
diff --git a/Engine/Plugins/Experimental/HairCardGenerator/Source/HairCardGeneratorDataflow/Public/GenerateCardsTexturesNode.h b/Engine/Plugins/Experimental/HairCardGenerator/Source/HairCardGeneratorDataflow/Public/GenerateCardsTexturesNode.h
--- a/Engine/Plugins/Experimental/HairCardGenerator/Source/HairCardGeneratorDataflow/Public/GenerateCardsTexturesNode.h
+++ b/Engine/Plugins/Experimental/HairCardGenerator/Source/HairCardGeneratorDataflow/Public/GenerateCardsTexturesNode.h
@@ -22,7 +22,7 @@ struct FCardsTextureSettings
 
 	// Total number of textures to generate for this LOD settings
 	UPROPERTY(EditAnywhere, Category="Groom Cards")
-	int32 NumTextures = 75;
+	int32 NumTextures = 100;
 };
 
 /** Generate the cards texture for final rendering */
