Change hair card gen default settings + assets path

#Synced-CL 51541525

diff --git a/Engine/Plugins/Experimental/HairCardGenerator/Source/HairCardGeneratorDataflow/Private/BuildCardsSettingsNode.cpp b/Engine/Plugins/Experimental/HairCardGenerator/Source/HairCardGeneratorDataflow/Private/BuildCardsSettingsNode.cpp
--- a/Engine/Plugins/Experimental/HairCardGenerator/Source/HairCardGeneratorDataflow/Private/BuildCardsSettingsNode.cpp
+++ b/Engine/Plugins/Experimental/HairCardGenerator/Source/HairCardGeneratorDataflow/Private/BuildCardsSettingsNode.cpp
@@ -61,6 +61,25 @@ FBuildCardsSettingsNode::FBuildCardsSettingsNode(const UE::Dataflow::FNodeParame
 	RegisterOutputConnection(&CardsSettings);
 }
 
+void FBuildCardsSettingsNode::OnPropertyChanged(UE::Dataflow::FContext& Context, const FPropertyChangedEvent& InPropertyChangedEvent)
+{
+	if (InPropertyChangedEvent.GetPropertyName() == GET_MEMBER_NAME_CHECKED(FBuildCardsSettingsNode, GenerationSettings) || 
+		InPropertyChangedEvent.GetPropertyName() == GET_MEMBER_NAME_CHECKED(FBuildCardsSettingsNode, GroomAsset))
+	{
+		if (UGroomAsset* LocalAsset = GetValue<TObjectPtr<UGroomAsset>>(Context, &GroomAsset))
+		{
+			if (AssetsDirectoryPath.IsEmpty())
+			{
+				AssetsDirectoryPath = HairCardSettings_Helpers::GetGroomDirectoryForHairCards(LocalAsset);
+			}
+			if (AssetsBaseName.IsEmpty())
+			{
+				AssetsBaseName = HairCardSettings_Helpers::GetGroomNameForHairCards(LocalAsset);
+			}
+		}
+	}
+}
+
 void FBuildCardsSettingsNode::Evaluate(UE::Dataflow::FContext& Context, const FDataflowOutput* Out) const
 {
 	if (Out->IsA<FManagedArrayCollection>(&Collection))
@@ -113,11 +132,13 @@ void FBuildCardsSettingsNode::Evaluate(UE::Dataflow::FContext& Context, const FD
 					AdvancedOptions.AtlasSize = static_cast<uint8>(LocalSettings.AdvancedOptions.AtlasSize);
 					AdvancedOptions.ReserveTextureSpaceLOD = LocalSettings.AdvancedOptions.ReserveTextureSpaceLOD;
 					AdvancedOptions.bUseGroomAssetStrandWidth = LocalSettings.AdvancedOptions.bUseGroomAssetStrandWidth;
+					AdvancedOptions.DefaultDirectory = AssetsDirectoryPath;
+					AdvancedOptions.DefaultName= AssetsBaseName;
 				
 					OutputSettings.AddDefaulted();
 					FHairCardGeneratorUtils::BuildGenerationSettings(false, LocalAsset, *CardsDescriptionPtr,
 						OutputSettings.Last().GenerationSettings, OutputSettings.Last().GenerationFlags, OutputSettings.Last().PipelineFlags, AdvancedOptions);
-
+					
 					if(!LocalSettings.FilterSettings.IsEmpty())
 					{
 						TArray<TArray<FName>> FilterCardGroups;
diff --git a/Engine/Plugins/Experimental/HairCardGenerator/Source/HairCardGeneratorDataflow/Public/BuildCardsSettingsNode.h b/Engine/Plugins/Experimental/HairCardGenerator/Source/HairCardGeneratorDataflow/Public/BuildCardsSettingsNode.h
--- a/Engine/Plugins/Experimental/HairCardGenerator/Source/HairCardGeneratorDataflow/Public/BuildCardsSettingsNode.h
+++ b/Engine/Plugins/Experimental/HairCardGenerator/Source/HairCardGeneratorDataflow/Public/BuildCardsSettingsNode.h
@@ -11,6 +11,15 @@
 
 #include "BuildCardsSettingsNode.generated.h"
 
+/** Default parameters for hair cards generation */
+namespace UE::HairCardsGen::Private
+{
+	static constexpr int32 NumClumps = 800;
+	static constexpr int32 NumFlyaways = 100;
+	static constexpr int32 NumTriangles = 45000;
+	static constexpr int32 NumTextures = 100;
+}
+
 USTRUCT()
 struct FGroomCardsSettings
 {
@@ -72,19 +81,19 @@ struct FGroomFilterSettings
 
 	// Total number of cards to generate for this LOD settings
 	UPROPERTY(EditAnywhere, Category="Groom Cards", meta=(ClampMin="1", ClampMax="100000", DisplayName = "Num Clumps"))
-	int32 NumClumps = 800;
+	int32 NumClumps = UE::HairCardsGen::Private::NumClumps;
 
 	// Maximum number of cards to assign to flyaway strands
 	UPROPERTY(EditAnywhere, Category="Groom Cards", meta=(ClampMin="0", ClampMax="1000"))
-	int32 NumFlyaways = 100;
+	int32 NumFlyaways = UE::HairCardsGen::Private::NumFlyaways;
 
 	// Total number of triangles to generate for this LOD settings
 	UPROPERTY(EditAnywhere, Category="Groom Cards", meta=(ClampMin="1", ClampMax="100000"))
-	int32 NumTriangles = 45000;
+	int32 NumTriangles = UE::HairCardsGen::Private::NumTriangles;
 
 	// Total number of textures to generate for this LOD settings
 	UPROPERTY(EditAnywhere, Category="Groom Cards")
-	int32 NumTextures = 100;
+	int32 NumTextures = UE::HairCardsGen::Private::NumTextures;
 
 	// Card group names that will belong to the filter group settings
 	UPROPERTY(EditAnywhere, Category="Groom Cards")
@@ -142,7 +151,7 @@ struct FGroomGenerationSettings
 	UPROPERTY(EditAnywhere, Category="Groom Cards")
 	int32 GroupIndex = 0;
 
-	// Lisdt of filter settings to use for that generation
+	// List of filter settings to use for that generation
 	UPROPERTY(EditAnywhere, Category="Groom Cards")
 	TArray<FGroomFilterSettings> FilterSettings;
 
@@ -172,6 +181,7 @@ private:
 	virtual bool CanDebugDrawViewMode(const FName& ViewModeName) const override {return true;}
 	virtual void DebugDraw(UE::Dataflow::FContext& Context, IDataflowDebugDrawInterface& DataflowRenderingInterface, const FDebugDrawParameters& DebugDrawParameters) const override;
 #endif
+	virtual void OnPropertyChanged(UE::Dataflow::FContext& Context, const FPropertyChangedEvent& InPropertyChangedEvent) override;
 	//~ End FDataflowNode interface
 
 	/** Groom asset to build the cards settings from */
@@ -189,5 +199,13 @@ private:
 	/** List of filter settings to override */
 	UPROPERTY(EditAnywhere, Category = "Groom Cards")
 	TArray<FGroomGenerationSettings> GenerationSettings;
+	
+	// Assets directory path used when generating the mesh/textures
+	UPROPERTY(EditAnywhere, AdvancedDisplay, Category="Groom Cards")
+	FString AssetsDirectoryPath = "";
+	
+	// Assets base name used when generating the mesh/textures
+	UPROPERTY(EditAnywhere, AdvancedDisplay, Category="Groom Cards")
+	FString AssetsBaseName = "";
 };
 
diff --git a/Engine/Plugins/Experimental/HairCardGenerator/Source/HairCardGeneratorDataflow/Public/GenerateCardsClumpsNode.h b/Engine/Plugins/Experimental/HairCardGenerator/Source/HairCardGeneratorDataflow/Public/GenerateCardsClumpsNode.h
--- a/Engine/Plugins/Experimental/HairCardGenerator/Source/HairCardGeneratorDataflow/Public/GenerateCardsClumpsNode.h
+++ b/Engine/Plugins/Experimental/HairCardGenerator/Source/HairCardGeneratorDataflow/Public/GenerateCardsClumpsNode.h
@@ -22,11 +22,11 @@ struct FCardsClumpsSettings
 
 	// Total number of cards to generate for this LOD settings
 	UPROPERTY(EditAnywhere, Category="Groom Cards", meta=(ClampMin="1", ClampMax="100000", DisplayName = "Num Clumps"))
-	int32 NumCards = 800;
+	int32 NumCards = UE::HairCardsGen::Private::NumClumps;
 
 	// Maximum number of cards to assign to flyaway strands
 	UPROPERTY(EditAnywhere, Category="Groom Cards", meta=(ClampMin="0", ClampMax="1000"))
-	int32 NumFlyaways = 100;
+	int32 NumFlyaways = UE::HairCardsGen::Private::NumFlyaways;
 };
 
 /** Generate the clumps used to build the cards from the strands */
diff --git a/Engine/Plugins/Experimental/HairCardGenerator/Source/HairCardGeneratorDataflow/Public/GenerateCardsGeometryNode.h b/Engine/Plugins/Experimental/HairCardGenerator/Source/HairCardGeneratorDataflow/Public/GenerateCardsGeometryNode.h
--- a/Engine/Plugins/Experimental/HairCardGenerator/Source/HairCardGeneratorDataflow/Public/GenerateCardsGeometryNode.h
+++ b/Engine/Plugins/Experimental/HairCardGenerator/Source/HairCardGeneratorDataflow/Public/GenerateCardsGeometryNode.h
@@ -22,7 +22,7 @@ struct FCardsGeometrySettings
 
 	// Total number of triangles to generate for this LOD settings
 	UPROPERTY(EditAnywhere, Category="Groom Cards", meta=(ClampMin="1", ClampMax="100000"))
-	int32 NumTriangles = 45000;
+	int32 NumTriangles = UE::HairCardsGen::Private::NumTriangles;
 };
 
 /** Generate the geometry used to build the cards from the strands */
diff --git a/Engine/Plugins/Experimental/HairCardGenerator/Source/HairCardGeneratorDataflow/Public/GenerateCardsTexturesNode.h b/Engine/Plugins/Experimental/HairCardGenerator/Source/HairCardGeneratorDataflow/Public/GenerateCardsTexturesNode.h
--- a/Engine/Plugins/Experimental/HairCardGenerator/Source/HairCardGeneratorDataflow/Public/GenerateCardsTexturesNode.h
+++ b/Engine/Plugins/Experimental/HairCardGenerator/Source/HairCardGeneratorDataflow/Public/GenerateCardsTexturesNode.h
@@ -22,7 +22,7 @@ struct FCardsTextureSettings
 
 	// Total number of textures to generate for this LOD settings
 	UPROPERTY(EditAnywhere, Category="Groom Cards")
-	int32 NumTextures = 100;
+	int32 NumTextures = UE::HairCardsGen::Private::NumTextures;
 };
 
 /** Generate the cards texture for final rendering */
diff --git a/Engine/Plugins/Experimental/HairCardGenerator/Source/HairCardGeneratorEditor/Private/HairCardGenOptionsWindow.cpp b/Engine/Plugins/Experimental/HairCardGenerator/Source/HairCardGeneratorEditor/Private/HairCardGenOptionsWindow.cpp
--- a/Engine/Plugins/Experimental/HairCardGenerator/Source/HairCardGeneratorEditor/Private/HairCardGenOptionsWindow.cpp
+++ b/Engine/Plugins/Experimental/HairCardGenerator/Source/HairCardGeneratorEditor/Private/HairCardGenOptionsWindow.cpp
@@ -720,7 +720,7 @@ FReply SHairCardGenOptionsWindow::OnCancelClick()
 
 FReply SHairCardGenOptionsWindow::OnResetToDefaultClick() const
 {
-	SettingsObject->ResetToDefault();
+	SettingsObject->ResetToDefault(FString(), FString());
 	DetailsView->SetObject(SettingsObject, true);
 	return FReply::Handled();
 }
diff --git a/Engine/Plugins/Experimental/HairCardGenerator/Source/HairCardGeneratorEditor/Private/HairCardGeneratorEditorModule.cpp b/Engine/Plugins/Experimental/HairCardGenerator/Source/HairCardGeneratorEditor/Private/HairCardGeneratorEditorModule.cpp
--- a/Engine/Plugins/Experimental/HairCardGenerator/Source/HairCardGeneratorEditor/Private/HairCardGeneratorEditorModule.cpp
+++ b/Engine/Plugins/Experimental/HairCardGenerator/Source/HairCardGeneratorEditor/Private/HairCardGeneratorEditorModule.cpp
@@ -94,7 +94,8 @@ namespace HairCardGeneratorEditor_Impl
 {
 	static TObjectPtr<UGroomAsset> ParseGroomAssetArg(const TArray<FString>& Args, bool bOptional);
 	static FHairGroupsCardsSourceDescription* FindCardSourceDescription(TObjectPtr<UGroomAsset> Groom, int32 LODIndex);
-	static TObjectPtr<UHairCardGeneratorPluginSettings> CreateGroomSettingsLast(TObjectPtr<UGroomAsset> Groom, const FHairGroupsCardsSourceDescription& SourceDesc);
+	static TObjectPtr<UHairCardGeneratorPluginSettings> CreateGroomSettingsLast(TObjectPtr<UGroomAsset> Groom, const FHairGroupsCardsSourceDescription& SourceDesc, 
+		const FString& DefaultDirectory = "", const FString& DefaultName = "");
 
 	static void DbgCmdCacheGroom(const TArray<FString>& Args);
 
@@ -136,7 +137,7 @@ bool FHairCardGeneratorUtils::BuildGenerationSettings(const bool bQuerySettings,
 	GenerationSettings = CardGenController->GetGroomSettings(NewGroomAsset, CardsDesc.LODIndex);
 	if ( !GenerationSettings )
 	{
-		GenerationSettings = HairCardGeneratorEditor_Impl::CreateGroomSettingsLast(NewGroomAsset, CardsDesc);
+		GenerationSettings = HairCardGeneratorEditor_Impl::CreateGroomSettingsLast(NewGroomAsset, CardsDesc, AdvancedOptions.DefaultDirectory, AdvancedOptions.DefaultName);
 	}
 	else
 	{
@@ -576,17 +577,18 @@ static FHairGroupsCardsSourceDescription* HairCardGeneratorEditor_Impl::FindCard
 	return nullptr;
 }
 
-static TObjectPtr<UHairCardGeneratorPluginSettings> HairCardGeneratorEditor_Impl::CreateGroomSettingsLast(TObjectPtr<UGroomAsset> Groom, const FHairGroupsCardsSourceDescription& SourceDesc)
+static TObjectPtr<UHairCardGeneratorPluginSettings> HairCardGeneratorEditor_Impl::CreateGroomSettingsLast(TObjectPtr<UGroomAsset> Groom, 
+	const FHairGroupsCardsSourceDescription& SourceDesc, const FString& DefaultDirectory, const FString& DefaultName)
 {
 	TObjectPtr<UHairCardGeneratorPluginSettings> NewSettings = NewObject<UHairCardGeneratorPluginSettings>(/*Outer =*/Groom);
 	NewSettings->SetSource(Groom, SourceDesc.LODIndex, SourceDesc.GroupIndex);
-	NewSettings->ResetToDefault();
+	NewSettings->ResetToDefault(DefaultDirectory, DefaultName);
 
 	// First try to reset from Json file (then fallback to previous settings if they exist)
-	if ( NewSettings->ResetFromSettingsJson() )
+	if ( NewSettings->ResetFromSettingsJson(DefaultDirectory, DefaultName) )
 		return NewSettings;
 
-	NewSettings->ResetFromSourceDescription(SourceDesc);
+	NewSettings->ResetFromSourceDescription(SourceDesc, DefaultDirectory, DefaultName);
 	return NewSettings;
 }
 
diff --git a/Engine/Plugins/Experimental/HairCardGenerator/Source/HairCardGeneratorEditor/Private/HairCardGeneratorPluginSettings.cpp b/Engine/Plugins/Experimental/HairCardGenerator/Source/HairCardGeneratorEditor/Private/HairCardGeneratorPluginSettings.cpp
--- a/Engine/Plugins/Experimental/HairCardGenerator/Source/HairCardGeneratorEditor/Private/HairCardGeneratorPluginSettings.cpp
+++ b/Engine/Plugins/Experimental/HairCardGenerator/Source/HairCardGeneratorEditor/Private/HairCardGeneratorPluginSettings.cpp
@@ -94,9 +94,6 @@ int FJsonHelper::GetIntFieldDefault(const TSharedPtr<FJsonObject>& Json, const F
  *****************************************************************************/
 namespace HairCardSettings_Helpers
 {
-    FString GetDefaultBaseNameForHairCards(TObjectPtr<const UGroomAsset> GroomAsset, int LODIndex);
-    FString GetDefaultPathForHairCards(TObjectPtr<const UGroomAsset> GroomAsset, int LODIndex);
-
 	void ResetSettingsFromTemplate(TObjectPtr<UObject> SettingsObject, TObjectPtr<const UObject> TemplateObject);
 	uint8 PipelineSettingsDiff(TObjectPtr<const UObject> SettingsObject, TObjectPtr<UObject> DiffObject);
 	bool IsWriteableSettingsProperty(FProperty* Property);
@@ -107,14 +104,18 @@ namespace HairCardSettings_Helpers
 /// @param GroomAsset The groom asset cards will be generated for
 /// @param LODIndex The LOD index for card generation (TODO: This should be removed when multi-lod supported)
 /// @return Base name for generated data (e.g. assets and intermediate data)
-FString HairCardSettings_Helpers::GetDefaultBaseNameForHairCards(TObjectPtr<const UGroomAsset> GroomAsset, int LODIndex)
+FString HairCardSettings_Helpers::GetDefaultBaseNameForHairCards(TObjectPtr<const UGroomAsset> GroomAsset, int LODIndex, const FString& DefaultName)
 {
     const UHairCardGeneratorEditorSettings* ConfigSettings = GetDefault<UHairCardGeneratorEditorSettings>();
 
     FFormatNamedArguments FormatArgs;
-    if (GroomAsset)
+	if (!DefaultName.IsEmpty())
+	{
+		FormatArgs.Add(TEXT("groom_name"), FText::FromString(DefaultName));
+	}
+	else if (GroomAsset)
     {
-        FormatArgs.Add(TEXT("groom_name"), FText::FromString(GroomAsset->GetName()));
+        FormatArgs.Add(TEXT("groom_name"), FText::FromString(GetGroomNameForHairCards(GroomAsset)));
     }
     FormatArgs.Add(TEXT("lod"), LODIndex);
 
@@ -126,23 +127,27 @@ FString HairCardSettings_Helpers::GetDefaultBaseNameForHairCards(TObjectPtr<cons
 /// @param GroomAsset The groom asset cards will be generated for
 /// @param LODIndex The LOD index for card generation (TODO: This should be removed when multi-lod supported)
 /// @return Content path for generated hair card assets
-FString HairCardSettings_Helpers::GetDefaultPathForHairCards(TObjectPtr<const UGroomAsset> GroomAsset, int LODIndex)
+FString HairCardSettings_Helpers::GetDefaultDestinationPathForHairCards(TObjectPtr<const UGroomAsset> GroomAsset, int LODIndex, const FString& DefaultDirectory, const FString& DefaultName)
 {
-    const UHairCardGeneratorEditorSettings* ConfigSettings = GetDefault<UHairCardGeneratorEditorSettings>();
-
-    FFormatNamedArguments FormatArgs;
-    if (GroomAsset)
-    {
-        FString PackageRoot, PackagePath, PackageName;
-        FPackageName::SplitLongPackageName(FPackageName::ObjectPathToPackageName(GroomAsset->GetPathName()),
-            PackageRoot, PackagePath, PackageName);
+	const UHairCardGeneratorEditorSettings* ConfigSettings = GetDefault<UHairCardGeneratorEditorSettings>();
 
-        FormatArgs.Add(TEXT("groom_dir"), FText::FromString(PackageRoot + PackagePath));
-        FormatArgs.Add(TEXT("groom_name"), FText::FromString(GroomAsset->GetName()));
-    }
-    else
+	FFormatNamedArguments FormatArgs;
+	if (!DefaultDirectory.IsEmpty())
+	{
+		FormatArgs.Add(TEXT("groom_dir"), FText::FromString(DefaultDirectory));
+	}
+	else
+	{
+		FormatArgs.Add(TEXT("groom_dir"), FText::FromString(GetGroomDirectoryForHairCards(GroomAsset)));
+	}
+	
+	if (!DefaultName.IsEmpty())
+	{
+		FormatArgs.Add(TEXT("groom_name"), FText::FromString(DefaultName));
+	}
+	else if (GroomAsset)
     {
-        FormatArgs.Add(TEXT("groom_dir"), FText::FromString(FPaths::ProjectContentDir()));
+        FormatArgs.Add(TEXT("groom_name"), FText::FromString(GetGroomNameForHairCards(GroomAsset)));
     }
     FormatArgs.Add(TEXT("lod"), LODIndex);
 
@@ -153,6 +158,24 @@ FString HairCardSettings_Helpers::GetDefaultPathForHairCards(TObjectPtr<const UG
     return DirPath;
 }
 
+FString HairCardSettings_Helpers::GetGroomDirectoryForHairCards(TObjectPtr<const UGroomAsset> GroomAsset)
+{
+	if (GroomAsset)
+	{
+		FString PackageRoot, PackagePath, PackageName;
+		FPackageName::SplitLongPackageName(FPackageName::ObjectPathToPackageName(GroomAsset->GetPathName()),
+			PackageRoot, PackagePath, PackageName);
+		
+		return PackageRoot + PackagePath;
+	}
+	return FPaths::ProjectContentDir();
+}
+
+FString HairCardSettings_Helpers::GetGroomNameForHairCards(TObjectPtr<const UGroomAsset> GroomAsset)
+{
+	return GroomAsset ? GroomAsset->GetName() : FString();
+}
+
 void HairCardSettings_Helpers::ResetSettingsFromTemplate(TObjectPtr<UObject> SettingsObject, TObjectPtr<const UObject> TemplateObject)
 {
 	if ( !TemplateObject->IsA(SettingsObject->GetClass()) )
@@ -325,7 +348,7 @@ void UHairCardGeneratorPluginSettings::SetSource(TObjectPtr<UGroomAsset> InSourc
 	GenerateForGroomGroup = PhysGroupIndex;
 }
 
-void UHairCardGeneratorPluginSettings::ResetFromSourceDescription(const FHairGroupsCardsSourceDescription& SourceDesc)
+void UHairCardGeneratorPluginSettings::ResetFromSourceDescription(const FHairGroupsCardsSourceDescription& SourceDesc, const FString& DefaultDirectory, const FString& DefaultName)
 {
 	TObjectPtr<const UHairCardGeneratorPluginSettings> LastSettings = Cast<UHairCardGeneratorPluginSettings>(SourceDesc.GenerationSettings);
 
@@ -346,8 +369,8 @@ void UHairCardGeneratorPluginSettings::ResetFromSourceDescription(const FHairGro
 	}
 	else
 	{
-		BaseFilename = HairCardSettings_Helpers::GetDefaultBaseNameForHairCards(GroomAsset, LODIndex);
-		DestinationPath.Path = HairCardSettings_Helpers::GetDefaultPathForHairCards(GroomAsset, LODIndex);
+		BaseFilename = HairCardSettings_Helpers::GetDefaultBaseNameForHairCards(GroomAsset, LODIndex, DefaultName);
+		DestinationPath.Path = HairCardSettings_Helpers::GetDefaultDestinationPathForHairCards(GroomAsset, LODIndex, DefaultDirectory, DefaultName);
 	}
 
 	TArray<int32> CardGroupCount;
@@ -373,13 +396,13 @@ void UHairCardGeneratorPluginSettings::ResetFromSourceDescription(const FHairGro
 	PostResetUpdates();
 }
 
-bool UHairCardGeneratorPluginSettings::ResetFromSettingsJson()
+bool UHairCardGeneratorPluginSettings::ResetFromSettingsJson(const FString& DefaultDirectory, const FString& DefaultName)
 {
 	// HACK: Need to bootstrap base filename and intermediate path in order to get json path
 	if ( BaseFilename.IsEmpty() )
 	{
-		BaseFilename = HairCardSettings_Helpers::GetDefaultBaseNameForHairCards(GroomAsset, LODIndex);
-		DestinationPath.Path = HairCardSettings_Helpers::GetDefaultPathForHairCards(GroomAsset, LODIndex);
+		BaseFilename = HairCardSettings_Helpers::GetDefaultBaseNameForHairCards(GroomAsset, LODIndex, DefaultName);
+		DestinationPath.Path = HairCardSettings_Helpers::GetDefaultDestinationPathForHairCards(GroomAsset, LODIndex, DefaultDirectory, DefaultName);
 		UpdateOutputPaths();
 		UpdateParentInfo();
 	}
@@ -395,13 +418,13 @@ bool UHairCardGeneratorPluginSettings::ResetFromSettingsJson()
 }
 
 
-void UHairCardGeneratorPluginSettings::ResetToDefault()
+void UHairCardGeneratorPluginSettings::ResetToDefault(const FString& DefaultDirectory, const FString& DefaultName)
 {
 	TObjectPtr<UHairCardGeneratorPluginSettings> DefaultSettings = UHairCardGeneratorPluginSettings::StaticClass()->GetDefaultObject<UHairCardGeneratorPluginSettings>();
 	HairCardSettings_Helpers::ResetSettingsFromTemplate(this, DefaultSettings);
 
-	BaseFilename = HairCardSettings_Helpers::GetDefaultBaseNameForHairCards(GroomAsset, LODIndex);
-	DestinationPath.Path = HairCardSettings_Helpers::GetDefaultPathForHairCards(GroomAsset, LODIndex);
+	BaseFilename = HairCardSettings_Helpers::GetDefaultBaseNameForHairCards(GroomAsset, LODIndex, DefaultName);
+	DestinationPath.Path = HairCardSettings_Helpers::GetDefaultDestinationPathForHairCards(GroomAsset, LODIndex, DefaultDirectory, DefaultName);
 
 	TArray<int32> CardGroupCount;
 	HairCardSettings_Helpers::GetStrandCardGroupInfo(GroomAsset, CardGroupIds, CardGroupCount);
@@ -1138,7 +1161,7 @@ void UHairCardGeneratorPluginSettings::PostEditChangeProperty(FPropertyChangedEv
 	{
 		if (BaseFilename.IsEmpty())
 		{
-			BaseFilename = HairCardSettings_Helpers::GetDefaultBaseNameForHairCards(GroomAsset, LODIndex);
+			BaseFilename = HairCardSettings_Helpers::GetDefaultBaseNameForHairCards(GroomAsset, LODIndex, FString());
 		}
 		else
 		{
@@ -1146,7 +1169,7 @@ void UHairCardGeneratorPluginSettings::PostEditChangeProperty(FPropertyChangedEv
 			BaseFilename = ObjectTools::SanitizeInvalidChars(BaseFilename, INVALID_OBJECTNAME_CHARACTERS TEXT("\\"));
 		}
 
-		if ( !ResetFromSettingsJson() )
+		if ( !ResetFromSettingsJson(FString(), FString()) )
 		{
 			UpdateOutputPaths();
 		}
@@ -1158,7 +1181,7 @@ void UHairCardGeneratorPluginSettings::PostEditChangeProperty(FPropertyChangedEv
 	{
 		if (DestinationPath.Path.IsEmpty())
 		{
-			DestinationPath.Path = HairCardSettings_Helpers::GetDefaultPathForHairCards(GroomAsset, LODIndex);
+			DestinationPath.Path = HairCardSettings_Helpers::GetDefaultDestinationPathForHairCards(GroomAsset, LODIndex, FString(), FString());
 		}
 	}
 	else if (PropertyName == GET_MEMBER_NAME_CHECKED(UHairCardGeneratorPluginSettings, bReduceCardsFromPreviousLOD) 
diff --git a/Engine/Plugins/Experimental/HairCardGenerator/Source/HairCardGeneratorEditor/Private/HairCardGeneratorPluginSettings.h b/Engine/Plugins/Experimental/HairCardGenerator/Source/HairCardGeneratorEditor/Private/HairCardGeneratorPluginSettings.h
--- a/Engine/Plugins/Experimental/HairCardGenerator/Source/HairCardGeneratorEditor/Private/HairCardGeneratorPluginSettings.h
+++ b/Engine/Plugins/Experimental/HairCardGenerator/Source/HairCardGeneratorEditor/Private/HairCardGeneratorPluginSettings.h
@@ -98,6 +98,14 @@ enum class EHairCardGenerationSettingsCategories : uint8
 				| (uint8)EHairCardGenerationPipeline::ImportTextures,
 };
 
+namespace HairCardSettings_Helpers
+{
+	HAIRCARDGENERATOREDITOR_API FString GetGroomDirectoryForHairCards(TObjectPtr<const UGroomAsset> GroomAsset);
+	HAIRCARDGENERATOREDITOR_API FString GetGroomNameForHairCards(TObjectPtr<const UGroomAsset> GroomAsset);
+	HAIRCARDGENERATOREDITOR_API FString GetDefaultBaseNameForHairCards(TObjectPtr<const UGroomAsset> GroomAsset, int LODIndex, const FString& DefaultName);
+	HAIRCARDGENERATOREDITOR_API FString GetDefaultDestinationPathForHairCards(TObjectPtr<const UGroomAsset> GroomAsset, int LODIndex, 
+						const FString& DefaultDirectory, const FString& DefaultName);
+}
 
 /**
  * Self contained hair card generation settings. Separated/Duplicated from the 
@@ -185,11 +193,11 @@ public:
 
 	void SetSource(TObjectPtr<UGroomAsset> SourceObject, int32 GenLODIndex = 0, int32 PhysGroupIndex = 0);
 	// Load last-run values for options from card source description or defaults if no last-run info
-	void ResetFromSourceDescription(const FHairGroupsCardsSourceDescription& SourceDesc);
+	void ResetFromSourceDescription(const FHairGroupsCardsSourceDescription& SourceDesc, const FString& DefaultDirectory, const FString& DefaultName);
 	// Load settings from Json file if one exists
-	bool ResetFromSettingsJson();
+	bool ResetFromSettingsJson(const FString& DefaultDirectory, const FString& DefaultName);
 	// Load default values for all options
-	void ResetToDefault();
+	void ResetToDefault(const FString& DefaultDirectory, const FString& DefaultName);
 	void ResetFilterGroupSettingsToDefault(int FilterGroupIndex);
 
 	// UI Hooks for adding/removing settings groups
diff --git a/Engine/Plugins/Experimental/HairCardGenerator/Source/HairCardGeneratorEditor/Public/HairCardGeneratorEditorModule.h b/Engine/Plugins/Experimental/HairCardGenerator/Source/HairCardGeneratorEditor/Public/HairCardGeneratorEditorModule.h
--- a/Engine/Plugins/Experimental/HairCardGenerator/Source/HairCardGeneratorEditor/Public/HairCardGeneratorEditorModule.h
+++ b/Engine/Plugins/Experimental/HairCardGenerator/Source/HairCardGeneratorEditor/Public/HairCardGeneratorEditorModule.h
@@ -31,6 +31,12 @@ struct HAIRCARDGENERATOREDITOR_API FCardsGenerationAdvancedOptions
 
 	// Use strand width
 	bool bUseGroomAssetStrandWidth = true;
+	
+	// Default directory to be used if defined
+	FString DefaultDirectory = "";
+	
+	// Default name to be used if defined
+	FString DefaultName = "";
 };
 
 /** FHairCardGeneratorUtils has static interface functions to be used in other modules */
