HairCardGenerator: Update to push load/save functionality to separate ufunctions, support direct wiring of ustruct types in dataflow cache (WIP). Fixup empty group runs to support new dataflow selection setup and fix clear mask for atlas.

#Synced-CL 52679804

diff --git a/Engine/Plugins/Experimental/HairCardGenerator/Source/HairCardGeneratorDataflow/Private/GenerateCardsClumpsNode.cpp b/Engine/Plugins/Experimental/HairCardGenerator/Source/HairCardGeneratorDataflow/Private/GenerateCardsClumpsNode.cpp
--- a/Engine/Plugins/Experimental/HairCardGenerator/Source/HairCardGeneratorDataflow/Private/GenerateCardsClumpsNode.cpp
+++ b/Engine/Plugins/Experimental/HairCardGenerator/Source/HairCardGeneratorDataflow/Private/GenerateCardsClumpsNode.cpp
@@ -75,21 +75,18 @@ void FGenerateCardsClumpsNode::Evaluate(UE::Dataflow::FContext& Context, const F
 						bool bHasClumps = FHairCardGeneratorUtils::RunCardsGeneration(GenerationSettings, LODSettings.PipelineFlags,
 				[&StrandsClumps, &NumClumps, &StrandsFilter](const TObjectPtr<const UHairCardGeneratorPluginSettings>& GenerationSettings, const int32 FilterIndex, const uint8 GenFlags)
 						{
-							TArray<int32> FilterClumps; int32 ClumpCount = 0;
-							if(FHairCardGeneratorUtils::GenerateCardsClumps(GenerationSettings, FilterIndex, GenFlags, FilterClumps, ClumpCount))
+							FHairCardClumpData ClumpData;
+							if(FHairCardGeneratorUtils::GenerateCardsClumps(GenerationSettings, FilterIndex, GenFlags, ClumpData))
 							{
-								int32 NumFilterClumps = 0;
-								for (int32 CurveIndex = 0, NumCurves = FilterClumps.Num()-1; CurveIndex < NumCurves; ++CurveIndex)
+								for (int32 CurveIndex = 0; CurveIndex < ClumpData.StrandLabels.Num(); ++CurveIndex)
 								{
-									if (StrandsClumps.IsValidIndex(CurveIndex) && (FilterClumps[CurveIndex] != INDEX_NONE))
+									if (StrandsClumps.IsValidIndex(CurveIndex) && ClumpData.StrandLabels[CurveIndex] != INDEX_NONE)
 									{
-										StrandsClumps[CurveIndex] = FilterClumps[CurveIndex] + NumClumps;
+										StrandsClumps[CurveIndex] = ClumpData.StrandLabels[CurveIndex] + NumClumps;
 										StrandsFilter[CurveIndex] = FilterIndex;
-										NumFilterClumps = FMath::Max(NumFilterClumps,  FilterClumps[CurveIndex]);
 									}
 								}
-								NumFilterClumps += 1;
-								NumClumps += NumFilterClumps;
+								NumClumps += ClumpData.MaxMainClump + 1;
 								return true;
 							}
 							return false;
diff --git a/Engine/Plugins/Experimental/HairCardGenerator/Source/HairCardGeneratorDataflow/Private/GenerateCardsGeometryNode.cpp b/Engine/Plugins/Experimental/HairCardGenerator/Source/HairCardGeneratorDataflow/Private/GenerateCardsGeometryNode.cpp
--- a/Engine/Plugins/Experimental/HairCardGenerator/Source/HairCardGeneratorDataflow/Private/GenerateCardsGeometryNode.cpp
+++ b/Engine/Plugins/Experimental/HairCardGenerator/Source/HairCardGeneratorDataflow/Private/GenerateCardsGeometryNode.cpp
@@ -62,44 +62,52 @@ void FGenerateCardsGeometryNode::Evaluate(UE::Dataflow::FContext& Context, const
 				{
 					if(FHairCardGeneratorUtils::LoadGenerationSettings(GenerationSettings))
 					{
-						TArray<TArray<FVector3f>> ClumpsGeometry;
-						
+						TArray<FIntVector3> ClumpsFaces;
+						TArray<FVector3f> ClumpsVertices;
+						TArray<int32> CardIndices;
+						int32 GlobalVertexOffset = 0;
+						int32 GlobalCardIndex = 0;
+
 						bool bHasGeometry = FHairCardGeneratorUtils::RunCardsGeneration(GenerationSettings, LODSettings.PipelineFlags,
-                        	[&ClumpsGeometry](const TObjectPtr<const UHairCardGeneratorPluginSettings>& GenerationSettings, const int32 FilterIndex, const uint8 GenFlags)
-						{
-                        	TArray<TArray<FVector3f>> FilterGeometry;
-                        	if(FHairCardGeneratorUtils::GenerateCardsGeometry(GenerationSettings, FilterIndex, GenFlags, FilterGeometry))
-                        	{
-								ClumpsGeometry.Append(FilterGeometry);
-                        		return true;
-                        	}
-							return false;
-						}, false);
-   
-						if(bHasGeometry)
+							[&ClumpsFaces, &ClumpsVertices, &CardIndices, &GlobalVertexOffset, &GlobalCardIndex]
+							(const TObjectPtr<const UHairCardGeneratorPluginSettings>& GenerationSettings, const int32 FilterIndex, const uint8 GenFlags)
 						{
-							TArray<FIntVector3> ClumpsFaces;
-							TArray<FVector3f> ClumpsVertices;
-							TArray<int32> CardIndices;
-							int32 VertexOffset = 0;
-   
-							int32 CardIndex = 0;
-							for(TArray<FVector3f>& GeometryVertices : ClumpsGeometry)
+							FHairCardGeomData GeomData;
+							FHairCardMeshData MeshData;
+							int32 CardCount = 0;
+							if(!FHairCardGeneratorUtils::GenerateCardsGeometry(GenerationSettings, FilterIndex, GenFlags, GeomData, MeshData, CardCount))
+								return false;
+
+							for (int32 c = 0; c < CardCount; ++c)
 							{
-								for(int32 CurveVertex = 0, NumVertices = GeometryVertices.Num()/2; CurveVertex < (NumVertices-1); ++CurveVertex)
+								const int32 v0 = MeshData.CardVertOffsets[c];
+								const int32 vc = MeshData.CardVertCounts[c];
+								const int32 f0 = MeshData.CardFaceOffsets[c];
+								const int32 fc = MeshData.CardFaceCounts[c];
+
+								for (int32 j = 0; j < vc; ++j)
 								{
-									const int32 VertexIndex = CurveVertex * 2 + VertexOffset;
-									ClumpsFaces.Add(FIntVector3(VertexIndex, VertexIndex + 2, VertexIndex + 3));
-									ClumpsFaces.Add(FIntVector3(VertexIndex + 3, VertexIndex + 1, VertexIndex));
+									const int32 Base = (v0 + j) * 3;
+									ClumpsVertices.Add(FVector3f(MeshData.Verts[Base], MeshData.Verts[Base+1], MeshData.Verts[Base+2]));
+									CardIndices.Add(GlobalCardIndex);
 								}
-								for(FVector3f& VertexPosition : GeometryVertices)
-                                {
-									ClumpsVertices.Add(VertexPosition);
-									CardIndices.Add(CardIndex);
+								// Face indices are card-local (0..vc-1); add GlobalVertexOffset for collection indices
+								for (int32 j = 0; j < fc; ++j)
+								{
+									const int32 Base = (f0 + j) * 3;
+									ClumpsFaces.Add(FIntVector3(
+										MeshData.Faces[Base]     + GlobalVertexOffset,
+										MeshData.Faces[Base + 1] + GlobalVertexOffset,
+										MeshData.Faces[Base + 2] + GlobalVertexOffset));
 								}
-								VertexOffset = ClumpsVertices.Num();
-								++CardIndex;
+								GlobalVertexOffset += vc;
+								++GlobalCardIndex;
 							}
+							return true;
+						}, false);
+
+						if(bHasGeometry)
+						{
 							
 							FString CardsVerticesLODGroup = CardsVerticesGroup.ToString();
 							CardsVerticesLODGroup.AppendInt(GenerationSettings->GetLODIndex());
diff --git a/Engine/Plugins/Experimental/HairCardGenerator/Source/HairCardGeneratorDataflow/Private/GenerateCardsTexturesNode.cpp b/Engine/Plugins/Experimental/HairCardGenerator/Source/HairCardGeneratorDataflow/Private/GenerateCardsTexturesNode.cpp
--- a/Engine/Plugins/Experimental/HairCardGenerator/Source/HairCardGeneratorDataflow/Private/GenerateCardsTexturesNode.cpp
+++ b/Engine/Plugins/Experimental/HairCardGenerator/Source/HairCardGeneratorDataflow/Private/GenerateCardsTexturesNode.cpp
@@ -60,76 +60,74 @@ void FGenerateCardsTexturesNode::Evaluate(UE::Dataflow::FContext& Context, const
 				{
 					if(FHairCardGeneratorUtils::LoadGenerationSettings(GenerationSettings))
 					{
-						TArray<int32> CardsTextures;
-						TArray<FVector2f> VertexUVs;
-						
+						// Accumulate per-group cluster data and card offsets across RunCardsGeneration groups
+						TArray<FHairCardTextureClusterData> AllClusterData;
+						TArray<int32> CardOffsets;
+						CardOffsets.Add(0);
+
 						bool bHasTextures = FHairCardGeneratorUtils::RunCardsGeneration(GenerationSettings, LODSettings.PipelineFlags,
-							[&CardsTextures, &VertexUVs](const TObjectPtr<const UHairCardGeneratorPluginSettings>& GenerationSettings, const int32 FilterIndex, const uint8 GenFlags)
+							[&AllClusterData, &CardOffsets](const TObjectPtr<const UHairCardGeneratorPluginSettings>& GenerationSettings, const int32 FilterIndex, const uint8 GenFlags)
 						{
-							TArray<int32> FilterTextures;
-							int32 TexturesCount = 0;
-							if(FHairCardGeneratorUtils::GenerateCardsTexturesClusters(GenerationSettings, FilterIndex, GenFlags, FilterTextures, TexturesCount))
-							{
-								if (!FilterTextures.IsEmpty())
-								{
-									const int32 CardCount = FilterTextures.Num() - 1 - TexturesCount;
-									const int32 CardOffset = CardsTextures.Num();
-
-									CardsTextures.Reserve(CardOffset + CardCount);
-									for (int32 CardIndex = 0; CardIndex < CardCount; ++CardIndex)
-									{
-										CardsTextures.Add((FilterTextures[CardIndex] < TexturesCount) ? FilterTextures[CardCount + FilterTextures[CardIndex]] + CardOffset : INDEX_NONE);
-									}
-									return true;
-								}
-							}
-							return false;
+							FHairCardTextureClusterData ClusterData;
+							if(!FHairCardGeneratorUtils::GenerateCardsTexturesClusters(GenerationSettings, FilterIndex, GenFlags, ClusterData))
+								return false;
+							CardOffsets.Add(CardOffsets.Last() + ClusterData.ClustersLabel.Num());
+							AllClusterData.Add(MoveTemp(ClusterData));
+							return true;
 						}, false);
-			
+
 						if(bHasTextures)
 						{
-							TArray<float> FilterUVs;
-							if(FHairCardGeneratorUtils::GenerateTexturesLayoutAndAtlases(GenerationSettings, LODSettings.PipelineFlags, FilterUVs))
+							FHairCardAtlasLayoutData LayoutData;
+							FHairCardAtlasUVData AtlasUVData;
+							if(FHairCardGeneratorUtils::GenerateTexturesLayoutAndAtlases(GenerationSettings, LODSettings.PipelineFlags, LayoutData, AtlasUVData))
 							{
-								if (!FilterUVs.IsEmpty())
+								// Build ObjectTextureIndices: per card, global center card ID (or INDEX_NONE for flyaways)
+								// Center cards satisfy ObjectTextureIndices[card] == card, enabling cluster vis in the renderer
+								TArray<int32> CardsTextures;
+								for(int32 Gid = 0; Gid < AllClusterData.Num(); ++Gid)
 								{
-									VertexUVs.Reserve(FilterUVs.Num());
-									int32 UVsIndex = 0;
-									while (UVsIndex < FilterUVs.Num())
+									const FHairCardTextureClusterData& CD = AllClusterData[Gid];
+									const int32 CardOffset = CardOffsets[Gid];
+									const int32 NumClusters = CD.ClustersCenterId.Num();
+									for(const int32 ClusterIdx : CD.ClustersLabel)
 									{
-										if ((FilterUVs[UVsIndex] != -1.0) && ((UVsIndex + 1) < FilterUVs.Num()))
-										{
-											VertexUVs.Add(FVector2f(FilterUVs[UVsIndex], FilterUVs[UVsIndex + 1]));
-											UVsIndex += 2;
-										}
-										else
-										{
-											UVsIndex += 1;
-										}
+										CardsTextures.Add((ClusterIdx >= 0 && ClusterIdx < NumClusters)
+											? CD.ClustersCenterId[ClusterIdx] + CardOffset
+											: INDEX_NONE);
 									}
-									FString CardsObjectsLODGroup = CardsObjectsGroup.ToString();
-									CardsObjectsLODGroup.AppendInt(GenerationSettings->GetLODIndex());
+								}
 
-									FString CardsVerticesLODGroup = FGenerateCardsGeometryNode::CardsVerticesGroup.ToString();
-									CardsVerticesLODGroup.AppendInt(GenerationSettings->GetLODIndex());
+								// Build VertexUVs directly from the flat AtlasUVData — no sentinel parsing needed
+								TArray<FVector2f> VertexUVs;
+								VertexUVs.Reserve(AtlasUVData.VertexUvs.Num() / 2);
+								for(int32 i = 0; i < AtlasUVData.VertexUvs.Num(); i += 2)
+								{
+									VertexUVs.Add(FVector2f(AtlasUVData.VertexUvs[i], AtlasUVData.VertexUvs[i + 1]));
+								}
 
-									TManagedArray<int32>& ObjectTextureIndices = GroomCollection.AddAttribute<int32>(ObjectTextureIndicesAttribute, FName(CardsObjectsLODGroup));
-									TManagedArray<FVector2f>& VertexTextureUvs = GroomCollection.AddAttribute<FVector2f>(VertexTextureUVsAttribute, FName(CardsVerticesLODGroup));
+								FString CardsObjectsLODGroup = CardsObjectsGroup.ToString();
+								CardsObjectsLODGroup.AppendInt(GenerationSettings->GetLODIndex());
 
-									GroomCollection.EmptyGroup(FName(CardsObjectsLODGroup));
-									GroomCollection.AddElements(CardsTextures.Num(), FName(CardsObjectsLODGroup));
+								FString CardsVerticesLODGroup = FGenerateCardsGeometryNode::CardsVerticesGroup.ToString();
+								CardsVerticesLODGroup.AppendInt(GenerationSettings->GetLODIndex());
 
-									for (int32 CardIndex = 0, NumCards = CardsTextures.Num(); CardIndex < NumCards; ++CardIndex)
-									{
-										ObjectTextureIndices[CardIndex] = CardsTextures[CardIndex];
-									}
+								TManagedArray<int32>& ObjectTextureIndices = GroomCollection.AddAttribute<int32>(ObjectTextureIndicesAttribute, FName(CardsObjectsLODGroup));
+								TManagedArray<FVector2f>& VertexTextureUvs = GroomCollection.AddAttribute<FVector2f>(VertexTextureUVsAttribute, FName(CardsVerticesLODGroup));
+
+								GroomCollection.EmptyGroup(FName(CardsObjectsLODGroup));
+								GroomCollection.AddElements(CardsTextures.Num(), FName(CardsObjectsLODGroup));
 
-									if (VertexUVs.Num() == VertexTextureUvs.Num())
+								for(int32 CardIndex = 0, NumCards = CardsTextures.Num(); CardIndex < NumCards; ++CardIndex)
+								{
+									ObjectTextureIndices[CardIndex] = CardsTextures[CardIndex];
+								}
+
+								if(VertexUVs.Num() == VertexTextureUvs.Num())
+								{
+									for(int32 VertexIndex = 0, NumVertices = VertexUVs.Num(); VertexIndex < NumVertices; ++VertexIndex)
 									{
-										for (int32 VertexIndex = 0, NumVertices = VertexUVs.Num(); VertexIndex < NumVertices; ++VertexIndex)
-										{
-											VertexTextureUvs[VertexIndex] = VertexUVs[VertexIndex];
-										}
+										VertexTextureUvs[VertexIndex] = VertexUVs[VertexIndex];
 									}
 								}
 							}
diff --git a/Engine/Plugins/Experimental/HairCardGenerator/Source/HairCardGeneratorEditor/Private/HairCardGenControllerBase.h b/Engine/Plugins/Experimental/HairCardGenerator/Source/HairCardGeneratorEditor/Private/HairCardGenControllerBase.h
--- a/Engine/Plugins/Experimental/HairCardGenerator/Source/HairCardGeneratorEditor/Private/HairCardGenControllerBase.h
+++ b/Engine/Plugins/Experimental/HairCardGenerator/Source/HairCardGeneratorEditor/Private/HairCardGenControllerBase.h
@@ -6,7 +6,7 @@
 #include "GroomCacheData.h"
 #include "Containers/Map.h"
 #include "Templates/Tuple.h"
-
+#include "HairCardGenTypes.h"
 
 #include "HairCardGenControllerBase.generated.h"
 
@@ -76,14 +76,58 @@ public:
 	UFUNCTION(BlueprintImplementableEvent, Category = Python)
 	bool LoadSettings(const UHairCardGeneratorPluginSettings* GeneratorSettings);
 
+	// ---- Per-group stage ufunctions ----
+
+	/** Pure compute. Returns clump assignment via OutClumpData (bool-first + pure-out). */
+	UFUNCTION(BlueprintImplementableEvent, Category = Python)
+	bool GenerateClumps(int32 GroupId, FHairCardClumpData& OutClumpData);
+
+	/** Pure compute. Builds in-memory geometry_opt state from clump data. */
 	UFUNCTION(BlueprintImplementableEvent, Category = Python)
-	TArray<float> GenerateCardsGeometry();
+	bool SetOptimizations(int32 GroupId, const FHairCardClumpData& ClumpData);
 
+	/** Pure compute. Returns card geometry via OutGeomData (bool-first + pure-out). */
 	UFUNCTION(BlueprintImplementableEvent, Category = Python)
-	TArray<int32> GenerateClumps(const int GroupId = -1);
+	bool GenerateCardsGeometry(FHairCardGeomData& OutGeomData, FHairCardMeshData& OutMeshData, int32& OutCardCount);
 
+	/** Pure compute. Returns texture cluster assignment via OutClusterData (bool-first + pure-out). */
 	UFUNCTION(BlueprintImplementableEvent, Category = Python)
-	TArray<int> SetOptimizations(const int GroupId = -1);
+	bool ClusterTextures(int32 GroupId, const FHairCardGeomData& GeomData, const FHairCardMeshData& MeshData,
+	                     FHairCardTextureClusterData& OutClusterData);
+
+	// ---- Cross-group stage ufunctions ----
+
+	/** Pure compute. Receives all-group data; returns atlas layout and per-vertex atlas UVs (bool-first + pure-out). */
+	UFUNCTION(BlueprintImplementableEvent, Category = Python)
+	bool GenerateTextureLayout(const TArray<FHairCardGeomData>& AllGeomData,
+	                           const TArray<FHairCardMeshData>& AllMeshData,
+	                           const TArray<FHairCardTextureClusterData>& AllClusterData,
+	                           FHairCardAtlasLayoutData& OutLayoutData,
+	                           FHairCardAtlasBinpackData& OutBinpackData,
+	                           FHairCardAtlasUVData& OutAtlasUVData);
+
+	/** Pure compute. Renders texture atlases to disk. */
+	UFUNCTION(BlueprintImplementableEvent, Category = Python)
+	bool GenerateTextureAtlases(const TArray<FHairCardGeomData>& AllGeomData,
+	                            const TArray<FHairCardMeshData>& AllMeshData,
+	                            const TArray<FHairCardTextureClusterData>& AllClusterData,
+	                            const FHairCardAtlasLayoutData& LayoutData,
+								const FHairCardAtlasBinpackData& BinpackData,
+	                            const float WidthScale = -1);
+
+	/** Pure compute. Writes the static mesh asset. */
+	UFUNCTION(BlueprintImplementableEvent, Category = Python)
+	bool GenerateMesh(UStaticMesh* StaticMesh,
+	                  const TArray<FHairCardGeomData>& AllGeomData,
+	                  const TArray<FHairCardMeshData>& AllMeshData,
+	                  const TArray<FHairCardTextureClusterData>& AllClusterData,
+	                  const FHairCardAtlasLayoutData& LayoutData);
+
+	// ---- Geometry generation internal functions ----
+
+	/** Returns num sub-cards per clump optimizer; available after SetOptimizations. */
+	UFUNCTION(BlueprintImplementableEvent, Category = Python)
+	TArray<int32> GetNumCardsPerClump();
 
 	UFUNCTION(BlueprintImplementableEvent, Category = Python)
 	TArray<float> GetAverageCurve(const int Id, const int Cid);
@@ -91,23 +135,37 @@ public:
 	UFUNCTION(BlueprintImplementableEvent, Category = Python)
 	void SetInterpolatedAvgCurve(const int Id, const int Cid, const TArray<float>& Points);
 
+	// ---- Save / Load ufunctions ----
+
+	UFUNCTION(BlueprintImplementableEvent, Category = Python)
+	bool SaveClumpData(int32 GroupId, const FHairCardClumpData& ClumpData);
+
+	UFUNCTION(BlueprintImplementableEvent, Category = Python)
+	bool LoadClumpData(int32 GroupId, FHairCardClumpData& OutClumpData);
+
+	UFUNCTION(BlueprintImplementableEvent, Category = Python)
+	bool SaveCardData(int32 GroupId, const FHairCardGeomData& GeometryData, const FHairCardMeshData& MeshData, int32 CardCount);
+
+	UFUNCTION(BlueprintImplementableEvent, Category = Python)
+	bool LoadCardData(int32 GroupId, FHairCardGeomData& OutGeomData, FHairCardMeshData& OutMeshData, int32& OutCardCount);
+
 	UFUNCTION(BlueprintImplementableEvent, Category = Python)
-	TArray<int32> ClusterTextures(int GroupIndex);
+	bool SaveTextureClusterData(int32 GroupId, const FHairCardTextureClusterData& TexClusterData);
 
 	UFUNCTION(BlueprintImplementableEvent, Category = Python)
-	TArray<float> GenerateTextureLayout();
+	bool LoadTextureClusterData(int32 GroupId, FHairCardTextureClusterData& OutTexClusterData);
 
 	UFUNCTION(BlueprintImplementableEvent, Category = Python)
-	bool GenerateTextureAtlases(const float WidthScale = -1);
+	bool SaveAtlasLayoutData(const FHairCardAtlasLayoutData& LayoutData);
 
 	UFUNCTION(BlueprintImplementableEvent, Category = Python)
-	bool GenerateMesh(UStaticMesh* StaticMesh);
+	bool LoadAtlasLayoutData(FHairCardAtlasLayoutData& OutLayoutData);
 
 	UFUNCTION(BlueprintImplementableEvent, Category = Python)
-	bool TestWriteFbxMesh();
+	bool SaveAtlasBinpackData(const FHairCardAtlasBinpackData& BinpackData);
 
 	UFUNCTION(BlueprintImplementableEvent, Category = Python)
-	bool CheckClusterTextures(int GroupIndex);
+	bool LoadAtlasBinpackData(FHairCardAtlasBinpackData& OutBinpackData);
 
 	// Helper C++ function for building a static mesh from verts/faces, etc.
 	UFUNCTION(BlueprintCallable, Category = Python)
diff --git a/Engine/Plugins/Experimental/HairCardGenerator/Source/HairCardGeneratorEditor/Private/HairCardGeneratorEditorModule.cpp b/Engine/Plugins/Experimental/HairCardGenerator/Source/HairCardGeneratorEditor/Private/HairCardGeneratorEditorModule.cpp
--- a/Engine/Plugins/Experimental/HairCardGenerator/Source/HairCardGeneratorEditor/Private/HairCardGeneratorEditorModule.cpp
+++ b/Engine/Plugins/Experimental/HairCardGenerator/Source/HairCardGeneratorEditor/Private/HairCardGeneratorEditorModule.cpp
@@ -8,6 +8,7 @@
 #include "HairCardGenOptionsWindow.h"
 #include "HairCardGeneratorPluginSettings.h"
 #include "HairCardGenControllerBase.h"
+#include "HairCardGenTypes.h"
 #include "HairCardGeneratorEditorSettings.h"
 #include "HairCardGenCardSubdivider.h"
 #include "HairCardGenStrandNCSInterpolator.h"
@@ -102,7 +103,7 @@ namespace HairCardGeneratorEditor_Impl
 	static bool LoadGroomData(UGroomAsset* GroomAsset, const bool SaveCached = false);
 	static bool BuildHairCardGenGroomData(const int PointsPerStrand, const FHairDescription& InGroomData, FHairCardGen_GroomData& Out);
 	static bool GenerateCardsForCardGroup(const TObjectPtr<const UHairCardGeneratorPluginSettings> GenerationSettings, const int32 FilterIndex, const uint8 GenFlags);
-	static bool GenerateCardsGeometry(TObjectPtr<const UHairCardGeneratorPluginSettings> Settings, int index, TArray<TArray<FVector3f>>& ClumpsGeometry);
+	static bool GenerateCardsGeometry(TObjectPtr<const UHairCardGeneratorPluginSettings> Settings, int index, FHairCardGeomData& OutGeomData, FHairCardMeshData& OutMeshData, int32& OutCardCount);
 	static void GetTextureImportNames(UHairCardGeneratorPluginSettings* Settings, FHairCardImportSettings& ImportSettings);
 	static TArray<UObject*> ImportHairCardTextures(const FString& ImportSrc, const FString& ImportDst, const FHairCardImportSettings& ImportSettings);
 	static UStaticMesh* CreateMeshInPackage(const FString& DestinationPackage, const FString& BaseFilename, int GroupID);
@@ -215,7 +216,7 @@ bool FHairCardGeneratorUtils::LoadGenerationSettings(const TObjectPtr<UHairCardG
 	return true;
 }
 
-bool FHairCardGeneratorUtils::GenerateCardsClumps(const TObjectPtr<const UHairCardGeneratorPluginSettings>& GenerationSettings, const int32 FilterIndex, const uint8 GenFlags, TArray<int32>& StrandsClumps, int32& NumClumps)
+bool FHairCardGeneratorUtils::GenerateCardsClumps(const TObjectPtr<const UHairCardGeneratorPluginSettings>& GenerationSettings, const int32 FilterIndex, const uint8 GenFlags, FHairCardClumpData& OutClumpData)
 {
 	TObjectPtr<UHairCardGenControllerBase> CardGenController = FHairCardGenController::GetHairCardGenerator();
 	if ( !CardGenController )
@@ -225,21 +226,20 @@ bool FHairCardGeneratorUtils::GenerateCardsClumps(const TObjectPtr<const UHairCa
 
 	if ( UHairCardGeneratorPluginSettings::CheckGenerationFlags(GenFlags, EHairCardGenerationPipeline::StrandClustering) && !GenerationSettings->bReduceCardsFromPreviousLOD )
 	{
-		StrandsClumps = CardGenController->GenerateClumps(FilterIndex);
-		if (StrandsClumps.IsEmpty())
+		if (!CardGenController->GenerateClumps(FilterIndex, OutClumpData))
 		{
 			UE_LOGF(LogHairCardGenerator, Error, "Strand clustering failed for settings group %d. Check the log for details.", FilterIndex);
 			return false;
 		}
-		else  
+		if (!CardGenController->SaveClumpData(FilterIndex, OutClumpData))
 		{
-			NumClumps = StrandsClumps.Last();
+			UE_LOG(LogHairCardGenerator, Warning, TEXT("Failed to save clump data for settings group %d."), FilterIndex);
 		}
 	}
 	return true;
 }
 
-bool FHairCardGeneratorUtils::GenerateCardsGeometry(const TObjectPtr<const UHairCardGeneratorPluginSettings>& GenerationSettings, const int32 FilterIndex, const uint8 GenFlags, TArray<TArray<FVector3f>>& ClumpsGeometry)
+bool FHairCardGeneratorUtils::GenerateCardsGeometry(const TObjectPtr<const UHairCardGeneratorPluginSettings>& GenerationSettings, const int32 FilterIndex, const uint8 GenFlags, FHairCardGeomData& OutGeomData, FHairCardMeshData& OutMeshData, int32& OutCardCount)
 {
 	TObjectPtr<UHairCardGenControllerBase> CardGenController = FHairCardGenController::GetHairCardGenerator();
 	if ( !CardGenController )
@@ -248,17 +248,21 @@ bool FHairCardGeneratorUtils::GenerateCardsGeometry(const TObjectPtr<const UHair
 	}
 	if ( UHairCardGeneratorPluginSettings::CheckGenerationFlags(GenFlags, EHairCardGenerationPipeline::GeometryGeneration) )
 	{
-		if ( !HairCardGeneratorEditor_Impl::GenerateCardsGeometry(GenerationSettings, FilterIndex,ClumpsGeometry) )
+		if ( !HairCardGeneratorEditor_Impl::GenerateCardsGeometry(GenerationSettings, FilterIndex, OutGeomData, OutMeshData, OutCardCount) )
 		{
 			UE_LOGF(LogHairCardGenerator, Error, "Cards geometry failed for settings group %d. Check the log for details.", FilterIndex)
 			return false;
 		}
+		if (!CardGenController->SaveCardData(FilterIndex, OutGeomData, OutMeshData, OutCardCount))
+		{
+			UE_LOG(LogHairCardGenerator, Warning, TEXT("Failed to save card data for settings group %d."), FilterIndex);
+		}
 	}
 	return true;
 }
 
 bool FHairCardGeneratorUtils::GenerateCardsTexturesClusters(const TObjectPtr<const UHairCardGeneratorPluginSettings>& GenerationSettings, const int32 FilterIndex, const uint8 GenFlags,
-	TArray<int32>& CardsTextures, int32& NumTextures)
+	FHairCardTextureClusterData& OutClusterData)
 {
 	TObjectPtr<UHairCardGenControllerBase> CardGenController = FHairCardGenController::GetHairCardGenerator();
 	if ( !CardGenController )
@@ -267,45 +271,92 @@ bool FHairCardGeneratorUtils::GenerateCardsTexturesClusters(const TObjectPtr<con
 	}
 	if ( UHairCardGeneratorPluginSettings::CheckGenerationFlags(GenFlags, EHairCardGenerationPipeline::TextureClustering) && !GenerationSettings->bReduceCardsFromPreviousLOD )
 	{
-		CardsTextures = CardGenController->ClusterTextures(FilterIndex);
-		if (CardsTextures.IsEmpty())
+		FHairCardGeomData GeomData;
+		FHairCardMeshData MeshData;
+		int32 CardCount = 0;
+		if (!CardGenController->LoadCardData(FilterIndex, GeomData, MeshData, CardCount))
+		{
+			UE_LOG(LogHairCardGenerator, Error, TEXT("Failed to load card data for settings group %d."), FilterIndex);
+			return false;
+		}
+
+		if (!CardGenController->ClusterTextures(FilterIndex, GeomData, MeshData, OutClusterData))
 		{
 			UE_LOGF(LogHairCardGenerator, Error, "Texture clustering failed for settings group %d. Check the log for details.", FilterIndex)
 			return false;
 		}
-		else
+		if (!CardGenController->SaveTextureClusterData(FilterIndex, OutClusterData))
 		{
-			NumTextures = CardsTextures.Last();
+			UE_LOG(LogHairCardGenerator, Warning, TEXT("Failed to save texture cluster data for settings group %d."), FilterIndex);
 		}
 	}
 	return true;
 }
 
-bool FHairCardGeneratorUtils::GenerateTexturesLayoutAndAtlases(const TObjectPtr<const UHairCardGeneratorPluginSettings>& GenerationSettings, const uint8 GenFlags, TArray<float>& VertexUVs)
+bool FHairCardGeneratorUtils::GenerateTexturesLayoutAndAtlases(const TObjectPtr<const UHairCardGeneratorPluginSettings>& GenerationSettings, const uint8 GenFlags, FHairCardAtlasLayoutData& OutLayoutData, FHairCardAtlasUVData& OutAtlasUVData)
 {
 	TObjectPtr<UHairCardGenControllerBase> CardGenController = FHairCardGenController::GetHairCardGenerator();
 	if ( !CardGenController )
 	{
 		return false;
 	}
-	if ( UHairCardGeneratorPluginSettings::CheckGenerationFlags(GenFlags, EHairCardGenerationPipeline::TextureLayout) && !GenerationSettings->bReduceCardsFromPreviousLOD )
+
+	const int32 NumGroups = GenerationSettings->GetFilterGroupSettings().Num();
+	TArray<FHairCardGeomData> AllGeom;
+	TArray<FHairCardMeshData> AllMesh;
+	TArray<FHairCardTextureClusterData> AllCluster;
+	TArray<int32> AllCardCount;
+	AllGeom.SetNum(NumGroups);
+	AllMesh.SetNum(NumGroups);
+	AllCluster.SetNum(NumGroups);
+	AllCardCount.SetNum(NumGroups);
+	for (int32 i = 0; i < NumGroups; ++i)
+	{
+		if (!CardGenController->LoadCardData(i, AllGeom[i], AllMesh[i], AllCardCount[i]))
+		{
+			UE_LOG(LogHairCardGenerator, Error, TEXT("Failed to load card data for group %d in texture layout step."), i);
+			return false;
+		}
+		if (!CardGenController->LoadTextureClusterData(i, AllCluster[i]))
+		{
+			UE_LOG(LogHairCardGenerator, Error, TEXT("Failed to load texture cluster data for group %d in texture layout step."), i);
+			return false;
+		}
+	}
+
+	// Layout is fast (bin-packing only) — always run it so OutLayoutData and OutAtlasUVData are
+	// reliably populated for callers (e.g. the Dataflow node repopulates VertexTextureUVs each eval).
+	// Skip only when deriving geometry from a previous LOD, where the parent layout is reused.
+	const bool bRunLayout    = !GenerationSettings->bReduceCardsFromPreviousLOD;
+	const bool bRunRendering = UHairCardGeneratorPluginSettings::CheckGenerationFlags(GenFlags, EHairCardGenerationPipeline::TextureRendering) && !GenerationSettings->bReduceCardsFromPreviousLOD;
+
+	if (bRunLayout)
 	{
-		VertexUVs = CardGenController->GenerateTextureLayout();
-		if ( VertexUVs.IsEmpty() )
+		FHairCardAtlasBinpackData BinpackData;
+		if (!CardGenController->GenerateTextureLayout(AllGeom, AllMesh, AllCluster, OutLayoutData, BinpackData, OutAtlasUVData))
 		{
 			UE_LOGF(LogHairCardGenerator, Error, "Texture layout failed. See log for details.");
 			return false;
 		}
+		CardGenController->SaveAtlasLayoutData(OutLayoutData);
+		CardGenController->SaveAtlasBinpackData(BinpackData);
 	}
-	
-	if ( UHairCardGeneratorPluginSettings::CheckGenerationFlags(GenFlags, EHairCardGenerationPipeline::TextureRendering) && !GenerationSettings->bReduceCardsFromPreviousLOD )
+
+	if (bRunRendering)
 	{
-		if ( !CardGenController->GenerateTextureAtlases() )
+		FHairCardAtlasBinpackData BinpackData;
+		if (!CardGenController->LoadAtlasBinpackData(BinpackData))
+		{
+			UE_LOG(LogHairCardGenerator, Error, TEXT("Failed to load atlas binpack data for texture rendering step."));
+			return false;
+		}
+		if (!CardGenController->GenerateTextureAtlases(AllGeom, AllMesh, AllCluster, OutLayoutData, BinpackData, -1.0f))
 		{
 			UE_LOGF(LogHairCardGenerator, Error, "Texture atlas creation failed. See log for details.");
 			return false;
 		}
 	}
+
 	return true;
 }
 
@@ -508,7 +559,35 @@ bool FHairCardGeneratorUtils::BuildCardsAssets(UGroomAsset* NewGroomAsset, FHair
 			return false;
 		}
 
-		if ( !CardGenController->GenerateMesh(NewMesh) )
+		// Load all per-group and layout data from disk for mesh generation
+		const int32 NumGroupsForMesh = GenerationSettings->GetFilterGroupSettings().Num();
+		TArray<FHairCardGeomData> MeshAllGeom;
+		TArray<FHairCardMeshData> MeshAllMesh;
+		TArray<FHairCardTextureClusterData> MeshAllCluster;
+		TArray<int32> MeshAllCardCount;
+		MeshAllGeom.SetNum(NumGroupsForMesh);
+		MeshAllMesh.SetNum(NumGroupsForMesh);
+		MeshAllCluster.SetNum(NumGroupsForMesh);
+		MeshAllCardCount.SetNum(NumGroupsForMesh);
+		for (int32 i = 0; i < NumGroupsForMesh; ++i)
+		{
+			if (!CardGenController->LoadCardData(i, MeshAllGeom[i], MeshAllMesh[i], MeshAllCardCount[i]) ||
+			    !CardGenController->LoadTextureClusterData(i, MeshAllCluster[i]))
+			{
+				UE_LOG(LogHairCardGenerator, Error, TEXT("Failed to load stage data for group %d in mesh generation step."), i);
+				HairCardGeneratorEditor_Impl::CreateFinishNotification(false, LOCTEXT("HairCardGen.FailedMeshGen", "Mesh generation failed. See log for details."));
+				return false;
+			}
+		}
+		FHairCardAtlasLayoutData MeshLayoutData;
+		if (!CardGenController->LoadAtlasLayoutData(MeshLayoutData))
+		{
+			UE_LOG(LogHairCardGenerator, Error, TEXT("Failed to load atlas layout data in mesh generation step."));
+			HairCardGeneratorEditor_Impl::CreateFinishNotification(false, LOCTEXT("HairCardGen.FailedMeshGen", "Mesh generation failed. See log for details."));
+			return false;
+		}
+
+		if ( !CardGenController->GenerateMesh(NewMesh, MeshAllGeom, MeshAllMesh, MeshAllCluster, MeshLayoutData) )
 		{
 			UE_LOGF(LogHairCardGenerator, Error, "Mesh generation failed. See log for details.");
 			HairCardGeneratorEditor_Impl::CreateFinishNotification(false, LOCTEXT("HairCardGen.FailedMeshGen", "Mesh generation failed. See log for details."));
@@ -641,7 +720,7 @@ static bool HairCardGeneratorEditor_Impl::BuildHairCardGenGroomData(const int Po
 	return GenerateHairCardGenGroomData(PointsPerStrand, InGroomData, Out, [](TArray<TArray<FVector>>&){});
 }
 
-static bool HairCardGeneratorEditor_Impl::GenerateCardsGeometry(TObjectPtr<const UHairCardGeneratorPluginSettings> Settings, int index, TArray<TArray<FVector3f>>& ClumpsGeometry)
+static bool HairCardGeneratorEditor_Impl::GenerateCardsGeometry(TObjectPtr<const UHairCardGeneratorPluginSettings> Settings, int index, FHairCardGeomData& OutGeomData, FHairCardMeshData& OutMeshData, int32& OutCardCount)
 {
 	TObjectPtr<UHairCardGenControllerBase> CardGenController = FHairCardGenController::GetHairCardGenerator();
 	if ( !CardGenController )
@@ -651,19 +730,39 @@ static bool HairCardGeneratorEditor_Impl::GenerateCardsGeometry(TObjectPtr<const
 
 	FScopedSlowTask SlowTask(10, FText::Format(LOCTEXT("GeneratingHairCards.GeometryOptimization", "Generating cards geometry"), index));
 
+	// Load clump data that was written by the clumps stage
 	SlowTask.EnterProgressFrame(1);
-	TArray<int> NumCards = CardGenController->SetOptimizations(index);
+	FHairCardClumpData ClumpData;
+	if (!CardGenController->LoadClumpData(index, ClumpData))
+	{
+		UE_LOG(LogHairCardGenerator, Error, TEXT("Failed to load clump data for settings group %d. Check the log for details."), index);
+		SlowTask.EnterProgressFrame(9);
+		return false;
+	}
 
-	if (NumCards.Num() == 0)
+	if (!CardGenController->SetOptimizations(index, ClumpData))
 	{
 		UE_LOGF(LogHairCardGenerator, Error, "Card geometry generation failed for settings group %d. Check the log for details.", index);
 		SlowTask.EnterProgressFrame(9);
 		return false;
 	}
 
+	TArray<int32> NumCards = CardGenController->GetNumCardsPerClump();
+
 	int TotalNumCards = 0;
 	for (int CurrentNumCards : NumCards) TotalNumCards += CurrentNumCards;
 
+	// Empty group: no cards to subdivide. Leave OutGeomData/OutMeshData default-constructed
+	// (all arrays empty) and OutCardCount = 0. SaveCardData will persist empty files so
+	// downstream LoadCardData succeeds unconditionally for all groups.
+	if (NumCards.IsEmpty() || TotalNumCards == 0)
+	{
+		UE_LOG(LogHairCardGenerator, Warning, TEXT("Settings group %d produced no cards — saving empty geometry data and continuing."), index);
+		SlowTask.EnterProgressFrame(9);
+		OutCardCount = 0;
+		return true;
+	}
+
 	float AvgLength = 0;
 	float AvgCurvRadius = 0;
 
@@ -714,7 +813,7 @@ static bool HairCardGeneratorEditor_Impl::GenerateCardsGeometry(TObjectPtr<const
 	}
 	else
 	{
-		auto CountTriangles = [&Subdivider, &AvgCurves]() 
+		auto CountTriangles = [&Subdivider, &AvgCurves]()
 		{
 			int TriCount = 0;
 			for (int CardId = 0; CardId < AvgCurves.Num(); CardId++) TriCount += 2 * (Subdivider.GetSubdivisionPoints(AvgCurves[CardId]).Num() / 3 - 1);
@@ -750,64 +849,42 @@ static bool HairCardGeneratorEditor_Impl::GenerateCardsGeometry(TObjectPtr<const
 	}
 
 	SlowTask.EnterProgressFrame(9);
-	TArray<float> ClumpsVertices = CardGenController->GenerateCardsGeometry();
-	if (ClumpsVertices.IsEmpty())
+	if (!CardGenController->GenerateCardsGeometry(OutGeomData, OutMeshData, OutCardCount))
 	{
 		UE_LOGF(LogHairCardGenerator, Error, "Card geometry generation failed for settings group %d. Check the log for details.", index);
 		return false;
 	}
-	
-	// Fill the vertices coming from the optimizer
-	ClumpsGeometry.Reset();
-	ClumpsGeometry.AddDefaulted();
-	
-	int32 VertexIndex = 0;
-	while(VertexIndex < ClumpsVertices.Num())
-	{
-		if(ClumpsVertices[VertexIndex] == -1.0)
-		{
-			if(VertexIndex != (ClumpsVertices.Num()-1))
-			{
-				ClumpsGeometry.AddDefaulted();
-			}
-			VertexIndex += 1;
-		}
-		else
-		{
-			ClumpsGeometry.Last().Add(FVector3f(ClumpsVertices[VertexIndex], ClumpsVertices[VertexIndex+1], ClumpsVertices[VertexIndex+2]));
-			VertexIndex += 3;
-		}
-	}
 
 	return true;
 }
 
+
 static bool HairCardGeneratorEditor_Impl::GenerateCardsForCardGroup(const TObjectPtr<const UHairCardGeneratorPluginSettings> GenerationSettings, const int32 FilterIndex, const uint8 GenFlags)
 {
 	FScopedSlowTask SlowTask(10, FText::Format(LOCTEXT("GeneratingHairCards.GenerateGroup", "Generating Hair Cards for settings group {0}..."), FilterIndex));
 	SlowTask.MakeDialog(/*bShowCancelButton =*/true);
 
 	SlowTask.EnterProgressFrame(1);
-	TArray<int32> StrandsClumps;
-	int32 NumClumps = 0;
-	if(!FHairCardGeneratorUtils::GenerateCardsClumps(GenerationSettings, FilterIndex, GenFlags, StrandsClumps, NumClumps))
+	FHairCardClumpData ClumpData;
+	if(!FHairCardGeneratorUtils::GenerateCardsClumps(GenerationSettings, FilterIndex, GenFlags, ClumpData))
 	{
 		SlowTask.EnterProgressFrame(9);
 		return false;
 	}
 
-	TArray<TArray<FVector3f>> ClumpsGeometry;
+	FHairCardGeomData GeomData;
+	FHairCardMeshData MeshData;
+	int32 CardCount = 0;
 	SlowTask.EnterProgressFrame(5);
-	if(!FHairCardGeneratorUtils::GenerateCardsGeometry(GenerationSettings, FilterIndex, GenFlags, ClumpsGeometry))
+	if(!FHairCardGeneratorUtils::GenerateCardsGeometry(GenerationSettings, FilterIndex, GenFlags, GeomData, MeshData, CardCount))
 	{
 		SlowTask.EnterProgressFrame(4);
 		return false;
 	}
 
 	SlowTask.EnterProgressFrame(4);
-	TArray<int32> CardsTextures;
-	int32 NumTextures= 0;
-	if(!FHairCardGeneratorUtils::GenerateCardsTexturesClusters(GenerationSettings, FilterIndex, GenFlags, CardsTextures, NumTextures))
+	FHairCardTextureClusterData ClusterData;
+	if(!FHairCardGeneratorUtils::GenerateCardsTexturesClusters(GenerationSettings, FilterIndex, GenFlags, ClusterData))
 	{
 		return false;
 	}
@@ -1006,8 +1083,9 @@ bool FHairCardGeneratorEditorModule::GenerateHairCardsForLOD(UGroomAsset* NewGro
 		return false;
 	}
 
-	TArray<float> VertexUVs;
-	if(!FHairCardGeneratorUtils::GenerateTexturesLayoutAndAtlases(GenerationSettings, AllGenFlags, VertexUVs))
+	FHairCardAtlasLayoutData LayoutData;
+	FHairCardAtlasUVData AtlasUVData;
+	if(!FHairCardGeneratorUtils::GenerateTexturesLayoutAndAtlases(GenerationSettings, AllGenFlags, LayoutData, AtlasUVData))
 	{
 		return false;
 	}
diff --git a/Engine/Plugins/Experimental/HairCardGenerator/Source/HairCardGeneratorEditor/Public/HairCardGenTypes.h b/Engine/Plugins/Experimental/HairCardGenerator/Source/HairCardGeneratorEditor/Public/HairCardGenTypes.h
new file mode 100644
--- /dev/null
+++ b/Engine/Plugins/Experimental/HairCardGenerator/Source/HairCardGeneratorEditor/Public/HairCardGenTypes.h
@@ -0,0 +1,203 @@
+// Copyright Epic Games, Inc. All Rights Reserved.
+#pragma once
+
+#include "CoreMinimal.h"
+#include "HairCardGenTypes.generated.h"
+
+/**
+ * Output of GenerateClumps / input to SetOptimizations.
+ * Per-settings-group; saved and loaded by SaveClumpData / LoadClumpData.
+ */
+USTRUCT(BlueprintType)
+struct FHairCardClumpData
+{
+	GENERATED_BODY()
+
+	/** (N_strands,) clump label index per strand */
+	UPROPERTY(BlueprintReadWrite, Category="Clustering")
+	TArray<int32> StrandLabels;
+
+	/** Largest main-clump index; used to separate flyaway from standard clumps in label list */
+	UPROPERTY(BlueprintReadWrite, Category="Clustering")
+	int32  MaxMainClump = 0;
+};
+
+
+/**
+ * Card mesh geometry for one settings group, stored as flat concatenated arrays.
+ * All cards are concatenated; use CardVertOffsets/CardFaceOffsets to slice per card.
+ * Embedded in FHairCardGeomData. Replaces per-card OBJ files in the Dataflow path (Layer 5B).
+ */
+USTRUCT(BlueprintType)
+struct FHairCardMeshData
+{
+	GENERATED_BODY()
+
+	/** (TotalVerts * 3,) flattened xyz — all cards concatenated */
+	UPROPERTY(BlueprintReadWrite, Category="Mesh")
+	TArray<float> Verts;
+
+	/** (TotalFaces * 3,) flattened vertex indices — card-local, use offsets to get global */
+	UPROPERTY(BlueprintReadWrite, Category="Mesh")
+	TArray<int32> Faces;
+
+	/** (TotalVerts * 2,) flattened uv — card surface parameterisation used for strand projection */
+	UPROPERTY(BlueprintReadWrite, Category="Mesh")
+	TArray<float> SurfaceUvs;
+
+	/** (N_cards,) start vert index per card into Verts/SurfaceUVs */
+	UPROPERTY(BlueprintReadWrite, Category="Mesh")
+	TArray<int32> CardVertOffsets;
+
+	/** (N_cards,) vert count per card */
+	UPROPERTY(BlueprintReadWrite, Category="Mesh")
+	TArray<int32> CardVertCounts;
+
+	/** (N_cards,) start face index per card into Faces */
+	UPROPERTY(BlueprintReadWrite, Category="Mesh")
+	TArray<int32> CardFaceOffsets;
+
+	/** (N_cards,) face count per card */
+	UPROPERTY(BlueprintReadWrite, Category="Mesh")
+	TArray<int32> CardFaceCounts;
+};
+
+
+/**
+ * Output of GenerateCardsGeometry / input to ClusterTextures and all cross-group stages.
+ * Per-settings-group; saved and loaded by SaveGeomData / LoadGeomData.
+ * CardCount is used by cross-group stages to compute global card ID offsets.
+ */
+USTRUCT(BlueprintType)
+struct FHairCardGeomData
+{
+	GENERATED_BODY()
+
+	/** (N_cards,) card length */
+	UPROPERTY(BlueprintReadWrite, Category="Geometry")
+	TArray<float> CardsLength;
+
+	/** (N_cards,) card width */
+	UPROPERTY(BlueprintReadWrite, Category="Geometry")
+	TArray<float> CardsWidth;
+
+	/** (N_cards,) left/right border ratio */
+	UPROPERTY(BlueprintReadWrite, Category="Geometry")
+	TArray<float> CardsLrRatio;
+
+	/** (N_cards * 3,) flattened card top-vector (used for orientation) */
+	UPROPERTY(BlueprintReadWrite, Category="Geometry")
+	TArray<float> CardsTopVec;
+
+	/**
+	 * (N_cards,) card type id.
+	 * Separates flyaway cards from standard cards for texture clustering.
+	 * Stored in card_type.npy (formerly cards_texture_group.npy).
+	 * 0 - standard card
+	 * 1 - flyaway card
+	 */
+	UPROPERTY(BlueprintReadWrite, Category="Geometry")
+	TArray<int32> CardsType;
+
+	/** (N_strands,) card index per strand (global card ID in cross-group context) */
+	UPROPERTY(BlueprintReadWrite, Category="Geometry")
+	TArray<int32> CardsLabel;
+
+	/** (N_cards,) physics group label per card */
+	UPROPERTY(BlueprintReadWrite, Category="Geometry")
+	TArray<uint8> CardsPhysGrpIdx;
+};
+
+
+/**
+ * Output of ClusterTextures / input to all cross-group stages.
+ * Per-settings-group; saved and loaded by SaveClusterData / LoadClusterData.
+ */
+USTRUCT(BlueprintType)
+struct FHairCardTextureClusterData
+{
+	GENERATED_BODY()
+
+	/** (N_cards,) cluster index per card */
+	UPROPERTY(BlueprintReadWrite, Category="TextureCluster")
+	TArray<int32> ClustersLabel;
+
+	/** (N_clusters,) center card index per cluster */
+	UPROPERTY(BlueprintReadWrite, Category="TextureCluster")
+	TArray<int32> ClustersCenterId;
+
+	/** (N_cards,) flip flag per card */
+	UPROPERTY(BlueprintReadWrite, Category="TextureCluster")
+	TArray<bool> CardsFlipped;
+};
+
+
+/**
+ * Texture atlas layout information for each texture to render
+ */
+USTRUCT(BlueprintType)
+struct FHairCardAtlasLayoutData
+{
+	GENERATED_BODY()
+
+	/** (N_centers * 2,) flattened [width, height] texture size for each cluster center */
+	UPROPERTY(BlueprintReadWrite, Category="TextureLayout")
+	TArray<int32> TxtSize;
+
+	/** (N_centers * 2,) flattened [x, y] atlas coordinate per cluster center */
+	UPROPERTY(BlueprintReadWrite, Category="TextureLayout")
+	TArray<int32> TxtCoords;
+
+	/** (N_centers,) global card id per cluster center */
+	UPROPERTY(BlueprintReadWrite, Category="TextureLayout")
+	TArray<int32> CenterCardsId;
+
+	// /** (N_centers * 2,) flattened [length, width] card size per cluster center */
+	// UPROPERTY()
+	// TArray<float> CenterCardsSize;
+};
+
+
+/**
+ * Binary partition info used internally for current atlas layout algorithm
+ */
+USTRUCT(BlueprintType)
+struct FHairCardAtlasBinpackData
+{
+	GENERATED_BODY()
+
+	/** (N_nodes * 2,) flattened child node indices; -1 = None (leaf) */
+	UPROPERTY(BlueprintReadWrite, Category="Binpack")
+	TArray<int32> Children;
+
+	/** (N_nodes * 4,) flattened [min_x, min_y, max_x, max_y] bounds per node */
+	UPROPERTY(BlueprintReadWrite, Category="Binpack")
+	TArray<int32> Bounds;
+
+	/** (N_nodes,) node full flag */
+	UPROPERTY(BlueprintReadWrite, Category="Binpack")
+	TArray<bool> IsFull;
+
+	/** (N_nodes,) node reserved flag */
+	UPROPERTY(BlueprintReadWrite, Category="Binpack")
+	TArray<bool> IsReserved;
+
+	/** Margin in pixels; invariant across all nodes */
+	UPROPERTY(BlueprintReadWrite, Category="Binpack")
+	int32 MarginPx = 0;
+};
+
+
+/**
+ * Per-vertex atlas UV coordinates, produced by GenerateTextureLayout alongside the layout data.
+ * Flat array parallel to FHairCardMeshData.SurfaceUvs — same vertex ordering, no sentinels.
+ */
+USTRUCT(BlueprintType)
+struct FHairCardAtlasUVData
+{
+	GENERATED_BODY()
+
+	/** (TotalVerts * 2,) flattened [u, v] atlas UV per vertex */
+	UPROPERTY(BlueprintReadWrite, Category="TextureLayout")
+	TArray<float> VertexUvs;
+};
diff --git a/Engine/Plugins/Experimental/HairCardGenerator/Source/HairCardGeneratorEditor/Public/HairCardGeneratorEditorModule.h b/Engine/Plugins/Experimental/HairCardGenerator/Source/HairCardGeneratorEditor/Public/HairCardGeneratorEditorModule.h
--- a/Engine/Plugins/Experimental/HairCardGenerator/Source/HairCardGeneratorEditor/Public/HairCardGeneratorEditorModule.h
+++ b/Engine/Plugins/Experimental/HairCardGenerator/Source/HairCardGeneratorEditor/Public/HairCardGeneratorEditorModule.h
@@ -2,6 +2,7 @@
 
 #pragma once
 #include "IHairCardGeneratorEditor.h"
+#include "HairCardGenTypes.h"
 #include <functional>
 
 struct FHairGroupsCardsSourceDescription;
@@ -53,16 +54,17 @@ struct HAIRCARDGENERATOREDITOR_API FHairCardGeneratorUtils
 	static bool LoadGenerationSettings(const TObjectPtr<UHairCardGeneratorPluginSettings>& GenerationSettings);
 
 	/** Generate the cards clumps based on the settings given a filter index */
-	static bool GenerateCardsClumps(const TObjectPtr<const UHairCardGeneratorPluginSettings>& GenerationSettings, const int32 FilterIndex, const uint8 GenFlags, TArray<int32>& StrandsClumps, int32& NumClumps);
+	static bool GenerateCardsClumps(const TObjectPtr<const UHairCardGeneratorPluginSettings>& GenerationSettings, const int32 FilterIndex, const uint8 GenFlags, FHairCardClumpData& OutClumpData);
 
 	/** Generate the cards geometry based on the settings given a filter index */
-	static bool GenerateCardsGeometry(const TObjectPtr<const UHairCardGeneratorPluginSettings>& GenerationSettings, const int32 FilterIndex, const uint8 GenFlags, TArray<TArray<FVector3f>>& ClumpsGeometry);
+	static bool GenerateCardsGeometry(const TObjectPtr<const UHairCardGeneratorPluginSettings>& GenerationSettings, const int32 FilterIndex, const uint8 GenFlags, FHairCardGeomData& OutGeomData, FHairCardMeshData& OutMeshData, int32& OutCardCount);
 
 	/** Generate the cards texture clusters based on the settings given a filter index */
-	static bool GenerateCardsTexturesClusters(const TObjectPtr<const UHairCardGeneratorPluginSettings>& GenerationSettings, const int32 FilterIndex, const uint8 GenFlags, TArray<int32>& CardsTextures, int32& NumTextures);
+	static bool GenerateCardsTexturesClusters(const TObjectPtr<const UHairCardGeneratorPluginSettings>& GenerationSettings, const int32 FilterIndex, const uint8 GenFlags, FHairCardTextureClusterData& OutClusterData);
 
-	/** Generate the cards texture layout and atlases */
-	static bool GenerateTexturesLayoutAndAtlases(const TObjectPtr<const UHairCardGeneratorPluginSettings>& GenerationSettings, const uint8 GenFlags, TArray<float>& VertexUVs);
+	/** Generate the cards texture layout and atlases. OutLayoutData is populated from layout generation or loaded from disk.
+	 *  OutAtlasUVData is populated when layout is generated; empty when layout is skipped (only rendering reruns). */
+	static bool GenerateTexturesLayoutAndAtlases(const TObjectPtr<const UHairCardGeneratorPluginSettings>& GenerationSettings, const uint8 GenFlags, FHairCardAtlasLayoutData& OutLayoutData, FHairCardAtlasUVData& OutAtlasUVData);
 	
 	/** Run the cards generation function */
 	static bool RunCardsGeneration(const TObjectPtr<UHairCardGeneratorPluginSettings>& GenerationSettings, const uint8 PipelineFlags, const FGeneratorFunction& PipelineFunction, const bool bCheckFlags = true);
