Improve cards template + fix terminal rendering

#virtualized
#Synced-CL 52686302

diff --git a/Engine/Plugins/Experimental/HairCardGenerator/Source/HairCardGeneratorDataflow/Private/ModifyCardsAttributesNode.cpp b/Engine/Plugins/Experimental/HairCardGenerator/Source/HairCardGeneratorDataflow/Private/ModifyCardsAttributesNode.cpp
--- a/Engine/Plugins/Experimental/HairCardGenerator/Source/HairCardGeneratorDataflow/Private/ModifyCardsAttributesNode.cpp
+++ b/Engine/Plugins/Experimental/HairCardGenerator/Source/HairCardGeneratorDataflow/Private/ModifyCardsAttributesNode.cpp
@@ -21,7 +21,7 @@ namespace UE::Groom::Private
 		
 		const TStrandAttributesConstRef<FName> CardGroupsAttribute = HairDescription.StrandAttributes().GetAttributesRef<FName>(HairAttribute::Strand::GroupCardsName);
 		
-		if (AttributesFacade.IsValid() && GeometryFacade.IsValid() && CardGroupsAttribute.IsValid())
+		if (AttributesFacade.IsValid() && GeometryFacade.IsValid())
 		{
 			FHairDescriptionGroups HairDescriptionGroups;
 			FGroomBuilder::BuildHairDescriptionGroups(HairDescription, HairDescriptionGroups, false);
@@ -40,10 +40,14 @@ namespace UE::Groom::Private
 					HairDescriptionGroups.HairGroups[GroupIndex].Strands.StrandsCurves.CurvesMapping.IsValidIndex(SourceCurve))
 				{
 					const int32 StrandID = HairDescriptionGroups.HairGroups[GroupIndex].Strands.StrandsCurves.CurvesMapping[SourceCurve];
-					if ((StrandID < CardGroupsAttribute.GetNumElements()) && (StrandID >= 0))
+					if (CardGroupsAttribute.IsValid() && (StrandID < CardGroupsAttribute.GetNumElements()) && (StrandID >= 0))
 					{
 						CardGroups[CurveIndex] = CardGroupsAttribute[FStrandID(StrandID)];
 					}
+					else
+					{
+						CardGroups[CurveIndex] = FName("Default");
+					}
 				}
 			}
 			AttributesFacade.SetCurveCardGroups(CardGroups);
diff --git a/Engine/Plugins/Experimental/HairCardGenerator/Source/HairCardGeneratorDataflow/Public/CardsAssetTerminalNode.h b/Engine/Plugins/Experimental/HairCardGenerator/Source/HairCardGeneratorDataflow/Public/CardsAssetTerminalNode.h
--- a/Engine/Plugins/Experimental/HairCardGenerator/Source/HairCardGeneratorDataflow/Public/CardsAssetTerminalNode.h
+++ b/Engine/Plugins/Experimental/HairCardGenerator/Source/HairCardGeneratorDataflow/Public/CardsAssetTerminalNode.h
@@ -28,7 +28,7 @@ private:
 	virtual void Evaluate(UE::Dataflow::FContext& Context, const FDataflowOutput* Out) const override;
 	//~ End FDataflowTerminalNode interface
 
-	UPROPERTY(meta = (DataflowInput, DataflowOutput, DisplayName = "Collection", DataflowPassthrough  = "Collection"))
+	UPROPERTY(meta = (DataflowInput, DataflowOutput, DisplayName = "Collection", DataflowPassthrough  = "Collection", DataflowRenderGroups = "CardsTexture"))
 	FManagedArrayCollection Collection;
 
 	/** Generator settings to be used */
